[AI Assistant] Fix toolset registry link errors in some build configurations

#Synced-CL 50267409

diff --git a/Engine/Plugins/Experimental/ToolsetRegistry/Source/ToolsetRegistry/Private/Tests/ToolsetRegistryFakeToolsetHandler.cpp b/Engine/Plugins/Experimental/ToolsetRegistry/Source/ToolsetRegistry/Private/Tests/ToolsetRegistryFakeToolsetHandler.cpp
--- a/Engine/Plugins/Experimental/ToolsetRegistry/Source/ToolsetRegistry/Private/Tests/ToolsetRegistryFakeToolsetHandler.cpp
+++ b/Engine/Plugins/Experimental/ToolsetRegistry/Source/ToolsetRegistry/Private/Tests/ToolsetRegistryFakeToolsetHandler.cpp
@@ -18,6 +18,8 @@ namespace UE::ToolsetRegistry
 					MakeValue(SUCCESSFUL_RESULT))));
 	}
 
+	FFakeToolset::~FFakeToolset() = default;
+
 	TFuture<TValueOrError<FString, FString>> FFakeToolset::ExecuteTool(
 		const FString& ToolName, const FString& JsonInput)
 	{
diff --git a/Engine/Plugins/Experimental/ToolsetRegistry/Source/ToolsetRegistry/Private/ToolsetRegistry.cpp b/Engine/Plugins/Experimental/ToolsetRegistry/Source/ToolsetRegistry/Private/ToolsetRegistry.cpp
--- a/Engine/Plugins/Experimental/ToolsetRegistry/Source/ToolsetRegistry/Private/ToolsetRegistry.cpp
+++ b/Engine/Plugins/Experimental/ToolsetRegistry/Source/ToolsetRegistry/Private/ToolsetRegistry.cpp
@@ -16,6 +16,9 @@
 
 namespace UE::ToolsetRegistry
 {
+	FToolsetRegistry::FToolsetRegistry() = default;
+	FToolsetRegistry::~FToolsetRegistry() = default;
+
 	TValueOrError<FToolDescriptor, FString> FToolDescriptor::FromString(
 		const FString& ToolsetAndToolName, bool LogErrors)
 	{
diff --git a/Engine/Plugins/Experimental/ToolsetRegistry/Source/ToolsetRegistry/Public/Tests/ToolsetRegistryFakeToolsetHandler.h b/Engine/Plugins/Experimental/ToolsetRegistry/Source/ToolsetRegistry/Public/Tests/ToolsetRegistryFakeToolsetHandler.h
--- a/Engine/Plugins/Experimental/ToolsetRegistry/Source/ToolsetRegistry/Public/Tests/ToolsetRegistryFakeToolsetHandler.h
+++ b/Engine/Plugins/Experimental/ToolsetRegistry/Source/ToolsetRegistry/Public/Tests/ToolsetRegistryFakeToolsetHandler.h
@@ -22,6 +22,7 @@ public:
 	static const inline FString SUCCESSFUL_RESULT = FString(TEXT(R"({"result": "success"})"));
 
 	FFakeToolset(const FString& Name);
+	virtual ~FFakeToolset();
 
 	TFuture<TValueOrError<FString, FString>> ExecuteTool(
 		const FString& ToolName, const FString& JsonInput) override;
diff --git a/Engine/Plugins/Experimental/ToolsetRegistry/Source/ToolsetRegistry/Public/ToolsetRegistry.h b/Engine/Plugins/Experimental/ToolsetRegistry/Source/ToolsetRegistry/Public/ToolsetRegistry.h
--- a/Engine/Plugins/Experimental/ToolsetRegistry/Source/ToolsetRegistry/Public/ToolsetRegistry.h
+++ b/Engine/Plugins/Experimental/ToolsetRegistry/Source/ToolsetRegistry/Public/ToolsetRegistry.h
@@ -27,6 +27,9 @@ namespace UE::ToolsetRegistry
 	// tools from toolsets.
 	class TOOLSETREGISTRY_API FToolsetRegistry {
 	public:
+		FToolsetRegistry();
+		~FToolsetRegistry();
+
 		// Register a toolset handler. Returns true if registration was successful.
 		bool RegisterToolset(TSharedPtr<IToolsetHandler> ToolsetHandler);
 
