[AI Assistant] Add namespacing from package metadata to toolset names to avoid registration conflicts

#rnx

#Synced-CL 50396013

diff --git a/Engine/Plugins/Experimental/ToolsetRegistry/Source/ToolsetRegistry/Private/BPToolsetRegistry.cpp b/Engine/Plugins/Experimental/ToolsetRegistry/Source/ToolsetRegistry/Private/BPToolsetRegistry.cpp
--- a/Engine/Plugins/Experimental/ToolsetRegistry/Source/ToolsetRegistry/Private/BPToolsetRegistry.cpp
+++ b/Engine/Plugins/Experimental/ToolsetRegistry/Source/ToolsetRegistry/Private/BPToolsetRegistry.cpp
@@ -24,7 +24,7 @@ void UToolsetRegistry::RegisterToolsetClass(UClass* InToolsetClass)
 	if (auto ToolsetRegistrySubsystem = UToolsetRegistrySubsystem::Get(TEXT("IsToolsetClassRegistered"));
 		ToolsetRegistrySubsystem.HasValue())
 	{
-		FString ToolsetName = InToolsetClass->GetName();
+		FString ToolsetName = UE::ToolsetRegistry::FToolsetLibraryWrapper::GetToolsetClassName(InToolsetClass);
 		TSharedPtr<UE::ToolsetRegistry::FToolsetLibraryWrapper> Wrapper =
 			MakeShared<UE::ToolsetRegistry::FToolsetLibraryWrapper>(InToolsetClass);
 
@@ -62,7 +62,7 @@ void UToolsetRegistry::UnregisterToolsetClass(UClass* InToolsetClass)
 			UToolsetRegistrySubsystem::Get(TEXT("IsToolsetClassRegistered"));
 		ToolsetRegistrySubsystem.HasValue())
 	{
-		FString ToolsetName = InToolsetClass->GetName();
+		FString ToolsetName = UE::ToolsetRegistry::FToolsetLibraryWrapper::GetToolsetClassName(InToolsetClass);
 		TSharedPtr<UE::ToolsetRegistry::IToolsetHandler> Handler =
 			ToolsetRegistrySubsystem.GetValue()->ToolsetRegistry.Find(ToolsetName);
 		if (Handler)
@@ -87,7 +87,7 @@ bool UToolsetRegistry::IsToolsetClassRegistered(UClass* InToolsetClass)
 	if (auto ToolsetRegistrySubsystem = UToolsetRegistrySubsystem::Get(TEXT("IsToolsetClassRegistered"));
 		ToolsetRegistrySubsystem.HasValue())
 	{
-		FString ToolsetName = InToolsetClass->GetName();
+		FString ToolsetName = UE::ToolsetRegistry::FToolsetLibraryWrapper::GetToolsetClassName(InToolsetClass);
 		return ToolsetRegistrySubsystem.GetValue()->ToolsetRegistry.Find(ToolsetName) != nullptr;
 	}
 	return false;
diff --git a/Engine/Plugins/Experimental/ToolsetRegistry/Source/ToolsetRegistry/Private/Tests/ToolsetRegistryToolsetTest.cpp b/Engine/Plugins/Experimental/ToolsetRegistry/Source/ToolsetRegistry/Private/Tests/ToolsetRegistryToolsetTest.cpp
--- a/Engine/Plugins/Experimental/ToolsetRegistry/Source/ToolsetRegistry/Private/Tests/ToolsetRegistryToolsetTest.cpp
+++ b/Engine/Plugins/Experimental/ToolsetRegistry/Source/ToolsetRegistry/Private/Tests/ToolsetRegistryToolsetTest.cpp
@@ -22,8 +22,7 @@
 #include "ToolsetRegistrySubsystem.h"
 
 #if WITH_DEV_AUTOMATION_TESTS
-const FString FakeToolsetName(TEXT("FakeToolset"));
-const FString FakeToolsetUpdatesVersionName(TEXT("FakeToolsetUpdatesVersion"));
+const FString FakeToolsetName(TEXT("ToolsetRegistry.FakeToolset"));
 const FString FakeToolsetInput(TEXT(R"({"inString":"hello"})"));
 const FString FakeToolName(TEXT("SimpleTool"));
 const FString ReturnValueName(TEXT("returnValue"));
@@ -49,7 +48,7 @@ class ToolsetRegistryContainer
 		auto ToolsetRegistrySubsystem = UToolsetRegistrySubsystem::Get(TEXT("Test Register"));
 		if (ToolsetRegistrySubsystem.HasValue() && ToolsetClass)
 		{
-			FString ToolsetName = ToolsetClass->GetName();
+			FString ToolsetName = UE::ToolsetRegistry::FToolsetLibraryWrapper::GetToolsetClassName(ToolsetClass);
 			Handler = ToolsetRegistrySubsystem.GetValue()->ToolsetRegistry.Find(ToolsetName);
 		}
 		return StaticCastSharedPtr<UE::ToolsetRegistry::FToolsetLibraryWrapper>(Handler);
@@ -61,8 +60,7 @@ class ToolsetRegistryContainer
 TSharedPtr<FJsonObject> GetFakeToolsetExpectedSchema()
 {
 	FJsonDomBuilder::FObject SchemaBuilder = MakeSchemaBuilder(
-		FString(TEXT("FakeToolset")), FString("1.0"), FString(TEXT("Fake Toolset that does nothing.")));
-
+		FString(FakeToolsetName), FString("1.0"), FString(TEXT("Fake Toolset that does nothing.")));
 
 	FJsonDomBuilder::FArray ToolsArray;
 	// Add SimpleTool
@@ -83,7 +81,7 @@ TSharedPtr<FJsonObject> GetFakeToolsetExpectedSchema()
 		AddRequiredListToSchema(OutputSchema, TArray<FString>({ReturnValueName}));
 
 		FJsonDomBuilder::FObject ToolEntry = MakeToolSchemaEntry(
-			TEXT("FakeToolset.SimpleTool"), TEXT("Simple tool that does nothing."),
+			TEXT("ToolsetRegistry.FakeToolset.SimpleTool"), TEXT("Simple tool that does nothing."),
 			InputSchema, OutputSchema);
 
 		AddToolSchema(SchemaBuilder, ToolsArray, ToolEntry);
@@ -108,7 +106,7 @@ TSharedPtr<FJsonObject> GetFakeToolsetExpectedSchema()
 		AddPropertiesToSchema(OutputSchema, OutputProperties);
 
 		FJsonDomBuilder::FObject ToolEntry = MakeToolSchemaEntry(
-			TEXT("FakeToolset.SimpleNoReturnTool"), TEXT("Simple tool that does nothing and returns nothing."),
+			TEXT("ToolsetRegistry.FakeToolset.SimpleNoReturnTool"), TEXT("Simple tool that does nothing and returns nothing."),
 			InputSchema, OutputSchema);
 
 		AddToolSchema(SchemaBuilder, ToolsArray, ToolEntry);
@@ -120,7 +118,7 @@ TSharedPtr<FJsonObject> GetFakeToolsetExpectedSchema()
 TSharedPtr<FJsonObject> GetFakeToolsetWithDefaultParamsExpectedSchema()
 {
 	FJsonDomBuilder::FObject SchemaBuilder = MakeSchemaBuilder(
-		FString(TEXT("FakeToolsetWithDefaultParams")), FString("1.0"),
+		FString(TEXT("ToolsetRegistry.FakeToolsetWithDefaultParams")), FString("1.0"),
 		FString(TEXT("Fake Toolset that has method with default parameters.")));
 
 	FJsonDomBuilder::FObject InputSchema = MakeSchema();
@@ -142,7 +140,7 @@ TSharedPtr<FJsonObject> GetFakeToolsetWithDefaultParamsExpectedSchema()
 	AddRequiredListToSchema(OutputSchema, TArray<FString>({ReturnValueName}));
 
 	FJsonDomBuilder::FObject ToolEntry = MakeToolSchemaEntry(
-		TEXT("FakeToolsetWithDefaultParams.ToolWithDefaults"),
+		TEXT("ToolsetRegistry.FakeToolsetWithDefaultParams.ToolWithDefaults"),
 		TEXT("Method with some params with default values."),
 		InputSchema, OutputSchema);
 
@@ -155,7 +153,7 @@ TSharedPtr<FJsonObject> GetFakeToolsetWithDefaultParamsExpectedSchema()
 TSharedPtr<FJsonObject> GetFakeToolsetWithArrayParamsExpectedSchema()
 {
 	FJsonDomBuilder::FObject SchemaBuilder = MakeSchemaBuilder(
-		FString(TEXT("FakeToolsetWithArrayParams")), FString("1.0"),
+		FString(TEXT("ToolsetRegistry.FakeToolsetWithArrayParams")), FString("1.0"),
 		FString(TEXT("Fake Toolset that has method with TArray parameters.")));
 
 	FJsonDomBuilder::FObject InputSchema = MakeSchema();
@@ -177,7 +175,7 @@ TSharedPtr<FJsonObject> GetFakeToolsetWithArrayParamsExpectedSchema()
 	AddRequiredListToSchema(OutputSchema, TArray<FString>({ReturnValueName}));
 
 	FJsonDomBuilder::FObject ToolEntry = MakeToolSchemaEntry(
-		TEXT("FakeToolsetWithArrayParams.ToolWithArrays"),
+		TEXT("ToolsetRegistry.FakeToolsetWithArrayParams.ToolWithArrays"),
 		TEXT("Method with some params with TArray values."),
 		InputSchema, OutputSchema);
 
@@ -190,7 +188,7 @@ TSharedPtr<FJsonObject> GetFakeToolsetWithArrayParamsExpectedSchema()
 TSharedPtr<FJsonObject> GetFakeToolsetWithSetParamsExpectedSchema()
 {
 	FJsonDomBuilder::FObject SchemaBuilder = MakeSchemaBuilder(
-		FString(TEXT("FakeToolsetWithSetParams")), FString("1.0"),
+		FString(TEXT("ToolsetRegistry.FakeToolsetWithSetParams")), FString("1.0"),
 		FString(TEXT("Fake Toolset that has method with TSet parameters.")));
 
 	FJsonDomBuilder::FObject InputSchema = MakeSchema();
@@ -212,7 +210,7 @@ TSharedPtr<FJsonObject> GetFakeToolsetWithSetParamsExpectedSchema()
 	AddRequiredListToSchema(OutputSchema, TArray<FString>({ReturnValueName}));
 
 	FJsonDomBuilder::FObject ToolEntry = MakeToolSchemaEntry(
-		TEXT("FakeToolsetWithSetParams.ToolWithSets"),
+		TEXT("ToolsetRegistry.FakeToolsetWithSetParams.ToolWithSets"),
 		TEXT("Method with some params with TSet values."),
 		InputSchema, OutputSchema);
 
@@ -225,7 +223,7 @@ TSharedPtr<FJsonObject> GetFakeToolsetWithSetParamsExpectedSchema()
 TSharedPtr<FJsonObject> GetFakeToolsetWithMapParamsExpectedSchema()
 {
 	FJsonDomBuilder::FObject SchemaBuilder = MakeSchemaBuilder(
-		FString(TEXT("FakeToolsetWithMapParams")), FString("1.0"),
+		FString(TEXT("ToolsetRegistry.FakeToolsetWithMapParams")), FString("1.0"),
 		FString(TEXT("Fake Toolset that has method with TMap parameters.")));
 
 	FJsonDomBuilder::FObject InputSchema = MakeSchema();
@@ -247,7 +245,7 @@ TSharedPtr<FJsonObject> GetFakeToolsetWithMapParamsExpectedSchema()
 	AddRequiredListToSchema(OutputSchema, TArray<FString>({ReturnValueName}));
 
 	FJsonDomBuilder::FObject ToolEntry = MakeToolSchemaEntry(
-		TEXT("FakeToolsetWithMapParams.ToolWithMaps"),
+		TEXT("ToolsetRegistry.FakeToolsetWithMapParams.ToolWithMaps"),
 		TEXT("Method with some params with TMap values."),
 		InputSchema, OutputSchema);
 
@@ -505,7 +503,7 @@ bool FToolsetRegistryToolsetTestExecuteToolInvalidToolsetName::RunTest(const FSt
 	FString Error;
 	AddExpectedError(TEXT("not found"), EAutomationExpectedErrorFlags::Contains);
 	UToolsetRegistry::ExecuteTool(
-		FString("FakeToolsetInvalid"), FakeToolName, FakeToolsetInput, Result, Error);
+		FString("ToolsetRegistry.FakeToolsetInvalid"), FakeToolName, FakeToolsetInput, Result, Error);
 	TestFalse("Tool execution with invalid toolset name returns error", Error.IsEmpty());
 	return true;
 }
@@ -540,6 +538,10 @@ bool FToolsetRegistryToolsetTestExecuteToolUnregisteredToolset::RunTest(const FS
 {
 	FString Result;
 	FString Error;
+
+	// Ensure it is unregistered
+	UToolsetRegistry::UnregisterToolsetClass(UFakeToolset::StaticClass());
+
 	AddExpectedError(TEXT("not found"), EAutomationExpectedErrorFlags::Contains);
 	UToolsetRegistry::ExecuteTool(
 		FakeToolsetName, FakeToolName, FakeToolsetInput, Result, Error);
@@ -586,7 +588,8 @@ bool FToolsetRegistryToolsetTestExecuteToolWithDefaultParam::RunTest(const FStri
 	FString Result;
 	FString Error;
 	UToolsetRegistry::ExecuteTool(
-		FString("FakeToolsetWithDefaultParams"), FString("ToolWithDefaults"), FakeToolsetInput, Result, Error);
+		FString("ToolsetRegistry.FakeToolsetWithDefaultParams"), FString("ToolWithDefaults"),
+		FakeToolsetInput, Result, Error);
 	TestTrue("Tool execution with default param returns without error", Error.IsEmpty());
 
 	FString ExpectedOutput = FString(TEXT(R"({"returnValue":"hello: 42"})"));
@@ -668,7 +671,8 @@ bool FToolsetRegistryToolsetTestExecutePythonToolWithError::RunTest(const FStrin
 	FString Input = TEXT(R"({"msg":"Hi"})");
 	FString Result;
 	FString Error;
-	UToolsetRegistry::ExecuteTool("ErrorProneToolset", "bad_tool", Input, Result, Error);
+	UToolsetRegistry::ExecuteTool(
+		"toolset_registry.tests.demo_toolset.ErrorProneToolset", "bad_tool", Input, Result, Error);
 	TestTrue("Tool was called correctly", Error.IsEmpty());
 
 	TSharedPtr<FJsonObject> ResultObject;
@@ -688,6 +692,44 @@ bool FToolsetRegistryToolsetTestExecutePythonToolWithError::RunTest(const FStrin
 	return true;
 }
 
+IMPLEMENT_SIMPLE_AUTOMATION_TEST(
+	FToolsetRegistryToolsetTestToolsetNaming,
+	"AI.ToolsetRegistry.Toolset.ToolsetNaming",
+	ToolsetRegistryTest::Flags)
+
+bool FToolsetRegistryToolsetTestToolsetNaming::RunTest(const FString& UnusedParameters)
+{
+	UClass* ToolsetClass = UFakeToolset::StaticClass();
+	FString ToolsetName = UE::ToolsetRegistry::FToolsetLibraryWrapper::GetToolsetClassName(ToolsetClass);
+	FString PackageName = ToolsetClass->GetOuter()->GetName();
+	FString PathName = ToolsetClass->GetPathName();
+	FString ClassName = ToolsetClass->GetName();
+	TestNotEqual("Toolset name should not be Class name", ToolsetName, ClassName);
+	TestNotEqual("Toolset name should not be same as Path name", ToolsetName, PathName);
+	TestTrue("Toolset name should include Class name", ToolsetName.Contains(ClassName));
+	int32 DotIndex;
+	TestTrue("Toolset name should include a namespace qualifier", ToolsetName.FindLastChar('.', DotIndex));
+	FString ToolsetNamespace = ToolsetName.Left(DotIndex);
+	FString ToolsetNameClass = ToolsetName.RightChop(DotIndex + 1);
+	TestEqual("Toolset name should have class name last", ToolsetNameClass, ClassName);
+	int32 SlashIndex;
+	TestTrue("Toolset package should be slash-separated", PackageName.FindLastChar('/', SlashIndex));
+	FString PackageLast = PackageName.RightChop(SlashIndex + 1);
+	TestEqual("Toolset name namespace should be rightmost package part", ToolsetNamespace, PackageLast);
+	return true;
+}
+
+IMPLEMENT_SIMPLE_AUTOMATION_TEST(
+	FToolsetRegistryToolsetTestToolsetNamingNull,
+	"AI.ToolsetRegistry.Toolset.ToolsetNamingNull",
+	ToolsetRegistryTest::Flags)
+
+bool FToolsetRegistryToolsetTestToolsetNamingNull::RunTest(const FString& UnusedParameters)
+{
+	FString ToolsetName = UE::ToolsetRegistry::FToolsetLibraryWrapper::GetToolsetClassName(nullptr);
+	return TestTrue("Toolset name for null UClass should be empty", ToolsetName.IsEmpty());
+}
+
 // Run python tests
 IMPLEMENT_SIMPLE_AUTOMATION_TEST(
 	FToolsetRegistryToolsetTestPythonTests,
diff --git a/Engine/Plugins/Experimental/ToolsetRegistry/Source/ToolsetRegistry/Private/ToolsetRegistryToolset.cpp b/Engine/Plugins/Experimental/ToolsetRegistry/Source/ToolsetRegistry/Private/ToolsetRegistryToolset.cpp
--- a/Engine/Plugins/Experimental/ToolsetRegistry/Source/ToolsetRegistry/Private/ToolsetRegistryToolset.cpp
+++ b/Engine/Plugins/Experimental/ToolsetRegistry/Source/ToolsetRegistry/Private/ToolsetRegistryToolset.cpp
@@ -54,6 +54,12 @@ namespace UE::ToolsetRegistry
 		};
 	}
 
+	FToolsetLibraryWrapper::FToolsetLibraryWrapper(UClass* LibraryClass) : ToolsetClass(LibraryClass)
+	{
+		GenerateToolCallObjects();
+		ToolsetName = ToolsetClass.Pin() ? FToolsetLibraryWrapper::GetToolsetClassName(LibraryClass) : FString();
+	}
+
 	TFuture<TValueOrError<FString, FString>> FToolsetLibraryWrapper::ExecuteTool(
 		const FString& ToolName, const FString& JsonInput)
 	{
@@ -107,6 +113,17 @@ namespace UE::ToolsetRegistry
 			TSharedPtr<FJsonObject> ToolEntry = FJsonSchemaGenerator::UStructToJsonSchemaObject(Function);
 			PRAGMA_ENABLE_DEPRECATION_WARNINGS
 
+			// Modify tool name from "ClassName.FunctionName" to "ToolsetName.FunctionName".
+			FString ToolName = ToolEntry->GetStringField(TEXT("name"));
+			int32 LastDotIndex;
+			// Assert if there is no dot, that would mean schema generation has failed catastrophically.
+			if (ensure(ToolName.FindLastChar('.', LastDotIndex)))
+			{
+				FString FunctionName = ToolName.RightChop(LastDotIndex + 1);
+				FString NewToolName = GetToolsetName() + "." + FunctionName;
+				ToolEntry->SetStringField(TEXT("name"), *NewToolName);
+			}
+
 			// Insert error field into outputSchema, creating it if it doesn't exist.
 			TSharedPtr<FJsonObject> OutputSchemaProperties;
 			TSharedPtr<FJsonObject> OutputSchema;
@@ -291,4 +308,46 @@ namespace UE::ToolsetRegistry
 		}
 		return ValidTools;
 	}
+
+	/*static*/ FString FToolsetLibraryWrapper::GetToolsetClassName(UClass* InClass)
+	{
+		FString Name;
+		if (InClass == nullptr)
+		{
+			return Name;
+		}
+		FString ClassName = InClass->GetName();
+		Name = ClassName;
+		check(InClass->GetOuter());
+		FString PackageName = InClass->GetOuter()->GetName();
+		int32 LastSlashIndex;
+		FString QualifierName;
+		if (ensure(PackageName.FindLastChar('/', LastSlashIndex)))
+		{
+			// Convert package name like "/Script/ToolsetRegistry" to a shorter qualifier like
+			// "ToolsetRegistry".
+			QualifierName = PackageName.RightChop(LastSlashIndex + 1);
+			if (QualifierName.EndsWith(FString("_PY")))
+			{
+				// Special handling for Python packages: convert package name like
+				// "/ToolsetRegistry/Python/toolset_registry/tests/demo_toolset_PY" as generated by
+				// PyUtil::GetGeneratedTypeOuterAndName, to "toolset_registry.tests.demo_toolset".
+				TArray<FString> PackageParts;
+				PackageName.ParseIntoArray(PackageParts, TEXT("/"));
+				int32 PythonPartIndex;
+				if (PackageParts.FindLast(TEXT("Python"), PythonPartIndex))
+				{
+					PackageParts.RemoveAt(0, PythonPartIndex + 1);
+				}
+				if (int32 ChopIndex = PackageParts.Last().Find(TEXT("_PY")); ChopIndex != INDEX_NONE)
+				{
+					PackageParts.Last() = PackageParts.Last().Left(ChopIndex);
+				}
+				QualifierName = FString::Join(PackageParts, TEXT("."));
+			}
+
+			Name = QualifierName + "." + ClassName;
+		}
+		return Name;
+	}
 }
diff --git a/Engine/Plugins/Experimental/ToolsetRegistry/Source/ToolsetRegistry/Public/ToolsetRegistryToolset.h b/Engine/Plugins/Experimental/ToolsetRegistry/Source/ToolsetRegistry/Public/ToolsetRegistryToolset.h
--- a/Engine/Plugins/Experimental/ToolsetRegistry/Source/ToolsetRegistry/Public/ToolsetRegistryToolset.h
+++ b/Engine/Plugins/Experimental/ToolsetRegistry/Source/ToolsetRegistry/Public/ToolsetRegistryToolset.h
@@ -19,11 +19,7 @@ namespace UE::ToolsetRegistry
 	class TOOLSETREGISTRY_API FToolsetLibraryWrapper : public IToolsetHandler
 	{
 	  public:
-		FToolsetLibraryWrapper(UClass* LibraryClass) : ToolsetClass(LibraryClass)
-		{
-			GenerateToolCallObjects();
-			ToolsetName = ToolsetClass.Pin() ? ToolsetClass->GetName() : FString();
-		}
+		FToolsetLibraryWrapper(UClass* LibraryClass);
 
 		// Executes a tool with the given name and JSON input, returning a future with the JSON
 		// output.
@@ -67,6 +63,9 @@ namespace UE::ToolsetRegistry
 		// Whether the toolset has a valid set of tools.
 		bool HasValidTools();
 
+		// Get qualified name for toolset class.
+		static FString GetToolsetClassName(UClass* InClass);
+
 	  private:
 
 		// The wrapped toolset function library
