Fix a crash in toolset invocation by allocating memory for any OutParms that apprear in the UFunction.

This is a temp fix until we can get guidance on how to better leverage existing ScriptCore functionality.

#Synced-CL 50390180

diff --git a/Engine/Plugins/Experimental/ToolsetRegistry/Source/ToolsetRegistry/Private/ToolsetRegistryToolCall.cpp b/Engine/Plugins/Experimental/ToolsetRegistry/Source/ToolsetRegistry/Private/ToolsetRegistryToolCall.cpp
--- a/Engine/Plugins/Experimental/ToolsetRegistry/Source/ToolsetRegistry/Private/ToolsetRegistryToolCall.cpp
+++ b/Engine/Plugins/Experimental/ToolsetRegistry/Source/ToolsetRegistry/Private/ToolsetRegistryToolCall.cpp
@@ -371,8 +371,48 @@ namespace UE::ToolsetRegistry
 		MaybeSetWorldContextPropertyInFrameMemory(Function, FrameMemory);
 
 		// Invoke function.
+		// TODO: A lot of this logic is duplicated from ScriptCore and should be ported to use existing function invocation logic.
 		FFrame StackFrame(InstanceObject, Function, FrameMemory, nullptr,
 			Function->ChildProperties);
+
+		// If the function has out parameters, fill the stack frame's out parameter info with the info for those params 
+		if (Function->HasAnyFunctionFlags(FUNC_HasOutParms))
+		{
+			FOutParmRec** LastOut = &StackFrame.OutParms;
+			for (FProperty* Property = (FProperty*)(Function->ChildProperties); Property && (Property->PropertyFlags & (CPF_Parm)) == CPF_Parm; Property = (FProperty*)Property->Next)
+			{
+				// this is used for optional parameters - the destination address for out parameter values is the address of the calling function
+				// so we'll need to know which address to use if we need to evaluate the default parm value expression located in the new function's
+				// bytecode
+				if (Property->HasAnyPropertyFlags(CPF_OutParm))
+				{
+					CA_SUPPRESS(6263)
+					FOutParmRec* Out = (FOutParmRec*)FMemory_Alloca(sizeof(FOutParmRec));
+					// set the address and property in the out param info
+					// note that since C++ doesn't support "optional out" we can ignore that here
+					Out->PropAddr = Property->ContainerPtrToValuePtr<uint8>(FrameMemory);
+					Out->Property = Property;
+
+					// add the new out param info to the stack frame's linked list
+					if (*LastOut)
+					{
+						(*LastOut)->NextOutParm = Out;
+						LastOut = &(*LastOut)->NextOutParm;
+					}
+					else
+					{
+						*LastOut = Out;
+					}
+				}
+			}
+
+			// set the next pointer of the last item to NULL to mark the end of the list
+			if (*LastOut)
+			{
+				(*LastOut)->NextOutParm = NULL;
+			}
+		}
+
 		const bool bHasReturnParam = (Function->ReturnValueOffset != MAX_uint16);
 		uint8* ReturnValueAddress = (bHasReturnParam ?
 			(FrameMemory + Function->ReturnValueOffset) : nullptr);
