Introducing ToolsetDefinition as the common base class for all Toolsets defined via UObject. This tightens the interface for what the BPRegistry can handle while avoiding BPFL which has some special technical and organizational meaning.

#Synced-CL 50463405

diff --git a/Engine/Plugins/Experimental/AIAssistant/Source/AIAssistant/Private/AIAssistantToolset.h b/Engine/Plugins/Experimental/AIAssistant/Source/AIAssistant/Private/AIAssistantToolset.h
--- a/Engine/Plugins/Experimental/AIAssistant/Source/AIAssistant/Private/AIAssistantToolset.h
+++ b/Engine/Plugins/Experimental/AIAssistant/Source/AIAssistant/Private/AIAssistantToolset.h
@@ -2,7 +2,7 @@
 
 #pragma once
 
-#include "Kismet/BlueprintFunctionLibrary.h"
+#include "ToolsetRegistry/Public/BPToolsetRegistry.h"
 
 #include "AIAssistantToolset.generated.h"
 
@@ -31,7 +31,7 @@ public:
 
 /// These functions help the AI assistant inspect its state as it relates to the larger Unreal Editor.
 UCLASS(BlueprintType)
-class UAIAssistantToolset : public UBlueprintFunctionLibrary
+class UAIAssistantToolset : public UToolsetDefinition
 {
 	GENERATED_BODY()
 public:
diff --git a/Engine/Plugins/Experimental/ToolsetRegistry/Source/ToolsetRegistry/Private/BPToolsetRegistry.cpp b/Engine/Plugins/Experimental/ToolsetRegistry/Source/ToolsetRegistry/Private/BPToolsetRegistry.cpp
--- a/Engine/Plugins/Experimental/ToolsetRegistry/Source/ToolsetRegistry/Private/BPToolsetRegistry.cpp
+++ b/Engine/Plugins/Experimental/ToolsetRegistry/Source/ToolsetRegistry/Private/BPToolsetRegistry.cpp
@@ -13,7 +13,7 @@
 #include "ToolsetRegistryModule.h"
 #include "ToolsetRegistrySubsystem.h"
 
-void UToolsetRegistry::RegisterToolsetClass(UClass* InToolsetClass)
+void UToolsetRegistry::RegisterToolsetClass(TSubclassOf<UToolsetDefinition> InToolsetClass)
 {
 	if (!InToolsetClass)
 	{
@@ -50,7 +50,7 @@ void UToolsetRegistry::RegisterToolsetClass(UClass* InToolsetClass)
 	}
 }
 
-void UToolsetRegistry::UnregisterToolsetClass(UClass* InToolsetClass)
+void UToolsetRegistry::UnregisterToolsetClass(TSubclassOf<UToolsetDefinition> InToolsetClass)
 {
 	if (!InToolsetClass)
 	{
@@ -77,7 +77,7 @@ void UToolsetRegistry::UnregisterToolsetClass(UClass* InToolsetClass)
 	}
 }
 
-bool UToolsetRegistry::IsToolsetClassRegistered(UClass* InToolsetClass)
+bool UToolsetRegistry::IsToolsetClassRegistered(TSubclassOf<UToolsetDefinition> InToolsetClass)
 {
 	if (!InToolsetClass)
 	{
@@ -136,7 +136,7 @@ void UToolsetRegistry::ExecuteTool(
 	}
 }
 
-FString UToolsetRegistry::GetToolsetJsonSchema(UClass* InToolsetClass)
+FString UToolsetRegistry::GetToolsetJsonSchema(TSubclassOf<UToolsetDefinition> InToolsetClass)
 {
 	if (!InToolsetClass)
 	{
diff --git a/Engine/Plugins/Experimental/ToolsetRegistry/Source/ToolsetRegistry/Private/Tests/ToolsetRegistryToolsetTest.cpp b/Engine/Plugins/Experimental/ToolsetRegistry/Source/ToolsetRegistry/Private/Tests/ToolsetRegistryToolsetTest.cpp
--- a/Engine/Plugins/Experimental/ToolsetRegistry/Source/ToolsetRegistry/Private/Tests/ToolsetRegistryToolsetTest.cpp
+++ b/Engine/Plugins/Experimental/ToolsetRegistry/Source/ToolsetRegistry/Private/Tests/ToolsetRegistryToolsetTest.cpp
@@ -31,7 +31,7 @@ const FString ErrorValueName(TEXT("error"));
 class ToolsetRegistryContainer
 {
   public:
-	ToolsetRegistryContainer(UClass* InToolsetClass)
+	ToolsetRegistryContainer(TSubclassOf<UToolsetDefinition> InToolsetClass)
 		: ToolsetClass(InToolsetClass)
 	{
 		UToolsetRegistry::RegisterToolsetClass(ToolsetClass);
@@ -54,7 +54,7 @@ class ToolsetRegistryContainer
 		return StaticCastSharedPtr<UE::ToolsetRegistry::FToolsetLibraryWrapper>(Handler);
 	}
 
-	UClass* ToolsetClass;
+	TSubclassOf<UToolsetDefinition> ToolsetClass;
 };
 
 TSharedPtr<FJsonObject> GetFakeToolsetExpectedSchema()
diff --git a/Engine/Plugins/Experimental/ToolsetRegistry/Source/ToolsetRegistry/Private/Tests/ToolsetRegistryToolsetTest.h b/Engine/Plugins/Experimental/ToolsetRegistry/Source/ToolsetRegistry/Private/Tests/ToolsetRegistryToolsetTest.h
--- a/Engine/Plugins/Experimental/ToolsetRegistry/Source/ToolsetRegistry/Private/Tests/ToolsetRegistryToolsetTest.h
+++ b/Engine/Plugins/Experimental/ToolsetRegistry/Source/ToolsetRegistry/Private/Tests/ToolsetRegistryToolsetTest.h
@@ -5,13 +5,13 @@
 #include "CoreMinimal.h"
 
 #include "Dom/JsonObject.h"
-#include "Kismet/BlueprintFunctionLibrary.h"
+#include "BPToolsetRegistry.h"
 
 #include "ToolsetRegistryToolsetTest.generated.h"
 
 /** Fake Toolset that does nothing. */
 UCLASS(Blueprintable)
-class UFakeToolset : public UBlueprintFunctionLibrary
+class UFakeToolset : public UToolsetDefinition
 {
 	GENERATED_BODY()
 
@@ -40,7 +40,7 @@ static TSharedPtr<FJsonObject> GetFakeToolsetExpectedSchema();
 
 /** Fake Toolset that has method with default parameters. */
 UCLASS(Blueprintable)
-class UFakeToolsetWithDefaultParams : public UBlueprintFunctionLibrary
+class UFakeToolsetWithDefaultParams : public UToolsetDefinition
 {
 	GENERATED_BODY()
 
@@ -59,7 +59,7 @@ static TSharedPtr<FJsonObject> GetFakeToolsetWithDefaultParamsExpectedSchema();
 
 /** Fake Toolset with no tools. */
 UCLASS(Blueprintable)
-class UFakeToolsetWithNoTools : public UBlueprintFunctionLibrary
+class UFakeToolsetWithNoTools : public UToolsetDefinition
 {
 	GENERATED_BODY()
 };
@@ -68,7 +68,7 @@ static TSharedPtr<FJsonObject> GetFakeToolsetWithNoToolsExpectedSchema();
 
 /** Fake Toolset that has method with TArray parameters. */
 UCLASS(Blueprintable)
-class UFakeToolsetWithArrayParams : public UBlueprintFunctionLibrary
+class UFakeToolsetWithArrayParams : public UToolsetDefinition
 {
 	GENERATED_BODY()
 
@@ -86,7 +86,7 @@ static TSharedPtr<FJsonObject> GetFakeToolsetWithArrayParamsExpectedSchema();
 
 /** Fake Toolset that has method with TMap parameters. */
 UCLASS(Blueprintable)
-class UFakeToolsetWithMapParams : public UBlueprintFunctionLibrary
+class UFakeToolsetWithMapParams : public UToolsetDefinition
 {
 	GENERATED_BODY()
 
@@ -104,7 +104,7 @@ static TSharedPtr<FJsonObject> GetFakeToolsetWithMapParamsExpectedSchema();
 
 /** Fake Toolset that has method with TSet parameters. */
 UCLASS(Blueprintable)
-class UFakeToolsetWithSetParams : public UBlueprintFunctionLibrary
+class UFakeToolsetWithSetParams : public UToolsetDefinition
 {
 	GENERATED_BODY()
 
@@ -122,7 +122,7 @@ static TSharedPtr<FJsonObject> GetFakeToolsetWithSetParamsExpectedSchema();
 
 /** Fake Toolset with private and non-static methods. */
 UCLASS(Blueprintable)
-class UFakeToolsetWithPrivateMethods : public UBlueprintFunctionLibrary
+class UFakeToolsetWithPrivateMethods : public UToolsetDefinition
 {
 	GENERATED_BODY()
 
diff --git a/Engine/Plugins/Experimental/ToolsetRegistry/Source/ToolsetRegistry/Public/BPToolsetRegistry.h b/Engine/Plugins/Experimental/ToolsetRegistry/Source/ToolsetRegistry/Public/BPToolsetRegistry.h
--- a/Engine/Plugins/Experimental/ToolsetRegistry/Source/ToolsetRegistry/Public/BPToolsetRegistry.h
+++ b/Engine/Plugins/Experimental/ToolsetRegistry/Source/ToolsetRegistry/Public/BPToolsetRegistry.h
@@ -4,9 +4,19 @@
 
 #include "CoreMinimal.h"
 #include "Kismet/BlueprintFunctionLibrary.h"
+#include "Templates/SubclassOf.h"
 
 #include "BPToolsetRegistry.generated.h"
 
+// This is the common base class for Toolsets defined as UObjects.
+// UFunctions that define tools in this class should be static functions.
+UCLASS(BlueprintType)
+class TOOLSETREGISTRY_API UToolsetDefinition : public UObject
+{
+	GENERATED_BODY()
+public:
+};
+
 // Blueprint-accessible wrapper to allow queries against the AI Toolset Registry.
 UCLASS()
 class TOOLSETREGISTRY_API UToolsetRegistry : public UBlueprintFunctionLibrary
@@ -17,15 +27,15 @@ public:
 
 	// Register a Blueprint function library class as a toolset.
 	UFUNCTION(BlueprintCallable, Category="ToolsetRegistry")
-	static void RegisterToolsetClass(UClass* InToolsetClass);
+	static void RegisterToolsetClass(TSubclassOf<UToolsetDefinition> InToolsetClass);
 
 	// Unregister a Blueprint function library class as a toolset.
 	UFUNCTION(BlueprintCallable, Category="ToolsetRegistry")
-	static void UnregisterToolsetClass(UClass* InToolsetClass);
+	static void UnregisterToolsetClass(TSubclassOf<UToolsetDefinition> InToolsetClass);
 
  	// Return whether a toolset class is already registered.
 	UFUNCTION(BlueprintPure, Category="ToolsetRegistry")
-	static bool IsToolsetClassRegistered(UClass* InToolsetClass);
+	static bool IsToolsetClassRegistered(TSubclassOf<UToolsetDefinition> InToolsetClass);
 
  	// Return whether a toolset is already registered by name.
 	UFUNCTION(BlueprintPure, Category="ToolsetRegistry")
@@ -40,5 +50,5 @@ public:
 	// Get toolset JSON schema for a toolset class.
 	// Meant for use while testing.
 	UFUNCTION(BlueprintPure, Category="ToolsetRegistry")
-	static FString GetToolsetJsonSchema(UClass* InToolsetClass);
+	static FString GetToolsetJsonSchema(TSubclassOf<UToolsetDefinition> InToolsetClass);
 };
