Adding core toolsets to ToolsetRegistry

#Synced-CL 50463405

diff --git a/Engine/Plugins/Experimental/ToolsetRegistry/Source/ToolsetRegistry/Private/ToolsetImage.cpp b/Engine/Plugins/Experimental/ToolsetRegistry/Source/ToolsetRegistry/Private/ToolsetImage.cpp
new file mode 100644
--- /dev/null
+++ b/Engine/Plugins/Experimental/ToolsetRegistry/Source/ToolsetRegistry/Private/ToolsetImage.cpp
@@ -0,0 +1,60 @@
+// Copyright Epic Games, Inc. All Rights Reserved.
+
+#include "ToolsetImage.h"
+
+#include "IImageWrapper.h"
+#include "IImageWrapperModule.h"
+#include "Misc/Base64.h"
+#include "Misc/FileHelper.h"
+#include "Modules/ModuleManager.h"
+
+#include UE_INLINE_GENERATED_CPP_BY_NAME(ToolsetImage)
+
+void EncodeImage(const TSharedPtr<IImageWrapper>& ImageWrapper, FString& OutData)
+{	
+	const TArray64<uint8>& CompressedData = ImageWrapper->GetCompressed();
+	OutData = FBase64::Encode(CompressedData.GetData(), CompressedData.Num());
+}
+
+bool FToolsetImage::SetFromBitmap(const TArray<FColor>& Bitmap, FIntPoint Dimensions)
+{
+	IImageWrapperModule& ImageWrapperModule = FModuleManager::LoadModuleChecked<IImageWrapperModule>("ImageWrapper");
+	TSharedPtr<IImageWrapper> ImageWrapper = ImageWrapperModule.CreateImageWrapper(EImageFormat::PNG);
+
+	if (ImageWrapper.IsValid() && 
+		ImageWrapper->SetRaw(Bitmap.GetData(), Bitmap.Num() * sizeof(FColor), Dimensions.X, Dimensions.Y, ERGBFormat::BGRA, 8))
+	{
+		EncodeImage(ImageWrapper, Data);
+		MimeType = "image/png";
+		return true;
+	}
+
+	return false;
+}
+
+bool FToolsetImage::SetFromFile(const FString& Path)
+{
+	TArray<uint8> ImageData;
+	if (!FFileHelper::LoadFileToArray(ImageData, *Path))
+	{
+		return false;
+	}
+
+	IImageWrapperModule& ImageWrapperModule = FModuleManager::LoadModuleChecked<IImageWrapperModule>("ImageWrapper");
+	
+	EImageFormat ImageFormat = ImageWrapperModule.DetectImageFormat(ImageData.GetData(), ImageData.Num());
+	if (ImageFormat != EImageFormat::PNG)
+	{
+		return false;
+	}
+
+	TSharedPtr<IImageWrapper> ImageWrapper = ImageWrapperModule.CreateImageWrapper(ImageFormat);
+	if (!ImageWrapper.IsValid() || !ImageWrapper->SetCompressed(ImageData.GetData(), ImageData.Num()))
+	{
+		return false;
+	}
+
+	EncodeImage(ImageWrapper, Data);
+	MimeType = "image/png";
+	return true;
+}
\ No newline at end of file
diff --git a/Engine/Plugins/Experimental/ToolsetRegistry/Source/ToolsetRegistry/Private/ToolsetLibrary.cpp b/Engine/Plugins/Experimental/ToolsetRegistry/Source/ToolsetRegistry/Private/ToolsetLibrary.cpp
new file mode 100644
--- /dev/null
+++ b/Engine/Plugins/Experimental/ToolsetRegistry/Source/ToolsetRegistry/Private/ToolsetLibrary.cpp
@@ -0,0 +1,81 @@
+// Copyright Epic Games, Inc. All Rights Reserved.
+
+#include "ToolsetLibrary.h"
+
+#include "JsonObjectConverter.h"
+#include "JsonSchema/JsonSchemaGenerator.h"
+#include "LevelEditorViewport.h"
+#include "UObject/PropertyAccessUtil.h"
+
+#include UE_INLINE_GENERATED_CPP_BY_NAME(ToolsetLibrary)
+
+FString UToolsetLibraryTemp::ListClassProperties(UClass* Class)
+{
+	FJsonSchemaPropertyFilter Filter;
+	Filter.CheckFlags |= CPF_BlueprintVisible;
+	Filter.SkipFlags |= CPF_Deprecated;
+	return FJsonSchemaGenerator::UStructToJsonSchemaString(Class, Filter);	
+}
+
+FString UToolsetLibraryTemp::GetObjectProperties(const UObject* Object, const TArray<FName>& PropertyNames)
+{
+	check(Object);
+	TSharedPtr<FJsonObject> PropertyData = MakeShared<FJsonObject>();
+	for (const auto& PropertyName : PropertyNames)
+	{
+		FProperty* Property = Object->GetClass()->FindPropertyByName(PropertyName);
+		if(Property)
+		{			
+			PropertyData->SetField(PropertyName.ToString(), FJsonObjectConverter::UPropertyToJsonValue(Property, Property->ContainerPtrToValuePtr<void>(Object)));
+		}
+	}
+	FString JsonProperties;
+	TSharedRef<TJsonWriter<>> Writer = TJsonWriterFactory<>::Create(&JsonProperties);
+	FJsonSerializer::Serialize(PropertyData.ToSharedRef(), Writer);
+	return JsonProperties;
+}
+
+bool UToolsetLibraryTemp::SetObjectProperties(UObject* Object, const FString& PropertiesJson)
+{
+	bool AllPropertiesSet = false;
+	TSharedPtr<FJsonObject> PropertyData;
+	TSharedRef<TJsonReader<>> Reader = TJsonReaderFactory<>::Create(PropertiesJson);
+	if (FJsonSerializer::Deserialize(Reader, PropertyData) && PropertyData.IsValid())
+	{
+		AllPropertiesSet = true;
+		const UClass* ObjectClass = Object->GetClass();
+		TSet<FString> Keys;
+		PropertyData->Values.GetKeys(Keys);
+		for (const auto& Key : Keys)
+		{
+			FProperty* Property = ObjectClass->FindPropertyByName(FName(Key));
+			if (Property)
+			{
+				TUniquePtr<FPropertyAccessChangeNotify> ChangeNotify = PropertyAccessUtil::BuildBasicChangeNotify(Property, Object, EPropertyAccessChangeNotifyMode::Always, EPropertyChangeType::ValueSet);
+				PropertyAccessUtil::EmitPreChangeNotify(ChangeNotify.Get(), false);
+				void* SettingsPropPtr = Property->ContainerPtrToValuePtr<void>(Object);
+				FJsonObjectConverter::JsonValueToUProperty(PropertyData->TryGetField(Key), Property, SettingsPropPtr);
+				PropertyAccessUtil::EmitPostChangeNotify(ChangeNotify.Get(), false);
+			}
+			else
+			{
+				AllPropertiesSet = false;
+			}
+		}
+	}
+	return AllPropertiesSet;
+}
+
+FString UToolsetLibraryTemp::ScreenshotPath(const FString& Filename)
+{
+	FString FullPath = Filename;
+	FScreenshotRequest::CreateViewportScreenShotFilename(FullPath);
+	return FullPath;
+}
+
+FToolsetImage UToolsetLibraryTemp::ScreenshotToImage(const FString& FullPath)
+{	
+	FToolsetImage Image;
+	Image.SetFromFile(FullPath);
+	return Image;
+}
diff --git a/Engine/Plugins/Experimental/ToolsetRegistry/Source/ToolsetRegistry/Public/ToolsetImage.h b/Engine/Plugins/Experimental/ToolsetRegistry/Source/ToolsetRegistry/Public/ToolsetImage.h
new file mode 100644
--- /dev/null
+++ b/Engine/Plugins/Experimental/ToolsetRegistry/Source/ToolsetRegistry/Public/ToolsetImage.h
@@ -0,0 +1,28 @@
+// Copyright Epic Games, Inc. All Rights Reserved. 
+
+#pragma once
+
+#include "CoreMinimal.h"
+
+#include "ToolsetImage.generated.h"
+
+/// The standard image format that toolsets should return.
+USTRUCT(BlueprintType)
+struct TOOLSETREGISTRY_API FToolsetImage
+{
+	GENERATED_BODY()
+public:
+	/// The format the image is encoded in.
+	UPROPERTY(BlueprintReadOnly, Category = "Image")
+	FString MimeType;
+
+	/// The image data encoded as a base64 string.
+	UPROPERTY(BlueprintReadOnly, Category = "Image")
+	FString Data;
+	
+	/// Populates the ToolsetImage based on a raw color image.
+	bool SetFromBitmap(const TArray<FColor>& Bitmap, FIntPoint Dimensions);
+
+	/// Populates the ToolsetImage from a file saved on disk.
+	bool SetFromFile(const FString& Path);
+};
diff --git a/Engine/Plugins/Experimental/ToolsetRegistry/Source/ToolsetRegistry/Public/ToolsetLibrary.h b/Engine/Plugins/Experimental/ToolsetRegistry/Source/ToolsetRegistry/Public/ToolsetLibrary.h
new file mode 100644
--- /dev/null
+++ b/Engine/Plugins/Experimental/ToolsetRegistry/Source/ToolsetRegistry/Public/ToolsetLibrary.h
@@ -0,0 +1,45 @@
+// Copyright Epic Games, Inc. All Rights Reserved. 
+
+#pragma once
+
+#include "Kismet/BlueprintFunctionLibrary.h"
+
+#include "ToolsetImage.h"
+
+#include "ToolsetLibrary.generated.h"
+
+UCLASS(BlueprintType)
+class TOOLSETREGISTRY_API UToolsetLibraryTemp : public UBlueprintFunctionLibrary
+{
+	GENERATED_BODY()
+public:
+
+	// TEMP JSON UTILS
+	// These functions are used by get/set/list properties but we should either move them to JsonUtilities
+	// or promote them to a core part of ToolsetRegistry (if we want to have them use custom JSON exporters).
+
+	/// Returns a JSON Schema formatted string that describes the Blueprint visible properties in the class.
+	UFUNCTION(BlueprintCallable, Category = "Toolset Library Temp")
+	static FString ListClassProperties(UClass* Class);
+
+	/// Returns a JSON formatted string that contains the values of the requested properties.
+	UFUNCTION(BlueprintCallable, Category = "Toolset Library Temp")
+	static FString GetObjectProperties(const UObject* Object, const TArray<FName>& PropertyNames);
+
+	/// Sets UObject properties to the values contained in the Json formatted string.
+	UFUNCTION(BlueprintCallable, Category = "Toolset Library Temp")
+	static bool SetObjectProperties(UObject* Object, const FString& PropertiesJson);
+
+	// TEMP SCREENSHOT UTILS
+	// These functions are workarounds for the lack of async tool call support. Once we have that,
+	// we should remove these functions and implement a tool that directly takes a screenshot and
+	// returns a ToolsetImage.
+
+	// Converts the file name of a screenshot into a full path based on UEs default logic.
+	UFUNCTION(BlueprintCallable, Category = "Toolset Library Temp")
+	static FString ScreenshotPath(const FString& Filename);
+
+	// Converts a screenshot image on disk to a ToolsetImage.
+	UFUNCTION(BlueprintCallable, Category = "Toolset Library Temp")
+	static FToolsetImage ScreenshotToImage(const FString& FullPath);
+};
