Fix deprecation warning.

#Synced-CL 50463405

diff --git a/Engine/Plugins/Experimental/ToolsetRegistry/Source/ToolsetRegistry/Private/ToolsetLibrary.cpp b/Engine/Plugins/Experimental/ToolsetRegistry/Source/ToolsetRegistry/Private/ToolsetLibrary.cpp
--- a/Engine/Plugins/Experimental/ToolsetRegistry/Source/ToolsetRegistry/Private/ToolsetLibrary.cpp
+++ b/Engine/Plugins/Experimental/ToolsetRegistry/Source/ToolsetRegistry/Private/ToolsetLibrary.cpp
@@ -14,7 +14,9 @@ FString UToolsetLibraryTemp::ListClassProperties(UClass* Class)
 	FJsonSchemaPropertyFilter Filter;
 	Filter.CheckFlags |= CPF_BlueprintVisible;
 	Filter.SkipFlags |= CPF_Deprecated;
+PRAGMA_DISABLE_DEPRECATION_WARNINGS
 	return FJsonSchemaGenerator::UStructToJsonSchemaString(Class, Filter);	
+PRAGMA_ENABLE_DEPRECATION_WARNINGS
 }
 
 FString UToolsetLibraryTemp::GetObjectProperties(const UObject* Object, const TArray<FName>& PropertyNames)
