ToolsetRegistry plugin should be Editor not EditorNoCommandlet

#Synced-CL 50463405

diff --git a/Engine/Plugins/Experimental/ToolsetRegistry/ToolsetRegistry.uplugin b/Engine/Plugins/Experimental/ToolsetRegistry/ToolsetRegistry.uplugin
--- a/Engine/Plugins/Experimental/ToolsetRegistry/ToolsetRegistry.uplugin
+++ b/Engine/Plugins/Experimental/ToolsetRegistry/ToolsetRegistry.uplugin
@@ -19,7 +19,7 @@
 	"Modules": [
 		{
 			"Name": "ToolsetRegistry",
-			"Type": "EditorNoCommandlet",
+			"Type": "Editor",
 			"LoadingPhase": "Default"
 		}
 	],
