Fix issues with error calling

#Synced-CL 50423628

diff --git a/Engine/Plugins/Experimental/AIAssistant/Source/AIAssistant/Private/AIAssistantToolResponse.cpp b/Engine/Plugins/Experimental/AIAssistant/Source/AIAssistant/Private/AIAssistantToolResponse.cpp
--- a/Engine/Plugins/Experimental/AIAssistant/Source/AIAssistant/Private/AIAssistantToolResponse.cpp
+++ b/Engine/Plugins/Experimental/AIAssistant/Source/AIAssistant/Private/AIAssistantToolResponse.cpp
@@ -35,6 +35,9 @@ namespace UE::AIAssistant
 		}
 		else
 		{
+			// The Backend currently expects a response to be present, even
+			// when an error occurs. This should be removed when that is fixed.
+			ToolResponse.ResponseRawJson = TEXT("{}");
 			ToolResponse.ErrorMessage = ToolResult.GetError();
 		}
 		return Options;
diff --git a/Engine/Plugins/Experimental/ToolsetRegistry/Source/ToolsetRegistry/Private/Tests/ToolsetRegistryToolsetTest.cpp b/Engine/Plugins/Experimental/ToolsetRegistry/Source/ToolsetRegistry/Private/Tests/ToolsetRegistryToolsetTest.cpp
--- a/Engine/Plugins/Experimental/ToolsetRegistry/Source/ToolsetRegistry/Private/Tests/ToolsetRegistryToolsetTest.cpp
+++ b/Engine/Plugins/Experimental/ToolsetRegistry/Source/ToolsetRegistry/Private/Tests/ToolsetRegistryToolsetTest.cpp
@@ -76,7 +76,6 @@ TSharedPtr<FJsonObject> GetFakeToolsetExpectedSchema()
 		FJsonDomBuilder::FObject OutputSchema = MakeSchema();
 		FJsonDomBuilder::FObject OutputProperties;
 		OutputProperties.Set(ReturnValueName, MakeProperty(TEXT("string")));
-		OutputProperties.Set(ErrorValueName, MakeProperty(TEXT("string")));
 		AddPropertiesToSchema(OutputSchema, OutputProperties);
 		AddRequiredListToSchema(OutputSchema, TArray<FString>({ReturnValueName}));
 
@@ -97,17 +96,10 @@ TSharedPtr<FJsonObject> GetFakeToolsetExpectedSchema()
 		AddPropertiesToSchema(InputSchema, Properties);
 		AddRequiredListToSchema(InputSchema, TArray<FString>({ParamName}));
 
-		FJsonDomBuilder::FObject OutputSchema = MakeSchema();
-		FJsonDomBuilder::FObject OutputProperties;
-		FJsonDomBuilder::FObject NullPropertyType;
-		NullPropertyType.Set(TEXT("type"), nullptr);
-		OutputProperties.Set(ReturnValueName, NullPropertyType);
-		OutputProperties.Set(ErrorValueName, MakeProperty(TEXT("string")));
-		AddPropertiesToSchema(OutputSchema, OutputProperties);
-
-		FJsonDomBuilder::FObject ToolEntry = MakeToolSchemaEntry(
-			TEXT("ToolsetRegistry.FakeToolset.SimpleNoReturnTool"), TEXT("Simple tool that does nothing and returns nothing."),
-			InputSchema, OutputSchema);
+		FJsonDomBuilder::FObject ToolEntry;
+		ToolEntry.Set(TEXT("description"), TEXT("Simple tool that does nothing and returns nothing."));
+		ToolEntry.Set(TEXT("name"), TEXT("ToolsetRegistry.FakeToolset.SimpleNoReturnTool"));
+		ToolEntry.Set(TEXT("inputSchema"), InputSchema);
 
 		AddToolSchema(SchemaBuilder, ToolsArray, ToolEntry);
 	}
@@ -135,7 +127,6 @@ TSharedPtr<FJsonObject> GetFakeToolsetWithDefaultParamsExpectedSchema()
 	FJsonDomBuilder::FObject OutputSchema = MakeSchema();
 	FJsonDomBuilder::FObject OutputProperties;
 	OutputProperties.Set(ReturnValueName, MakeProperty(TEXT("string")));
-	OutputProperties.Set(ErrorValueName, MakeProperty(TEXT("string")));
 	AddPropertiesToSchema(OutputSchema, OutputProperties);
 	AddRequiredListToSchema(OutputSchema, TArray<FString>({ReturnValueName}));
 
@@ -170,7 +161,6 @@ TSharedPtr<FJsonObject> GetFakeToolsetWithArrayParamsExpectedSchema()
 	FJsonDomBuilder::FObject OutputSchema = MakeSchema();
 	FJsonDomBuilder::FObject OutputProperties;
 	OutputProperties.Set(ReturnValueName, MakeArrayProperty(TEXT("number")));
-	OutputProperties.Set(ErrorValueName, MakeProperty(TEXT("string")));
 	AddPropertiesToSchema(OutputSchema, OutputProperties);
 	AddRequiredListToSchema(OutputSchema, TArray<FString>({ReturnValueName}));
 
@@ -205,7 +195,6 @@ TSharedPtr<FJsonObject> GetFakeToolsetWithSetParamsExpectedSchema()
 	FJsonDomBuilder::FObject OutputSchema = MakeSchema();
 	FJsonDomBuilder::FObject OutputProperties;
 	OutputProperties.Set(ReturnValueName, MakeSetProperty(TEXT("number")));
-	OutputProperties.Set(ErrorValueName, MakeProperty(TEXT("string")));
 	AddPropertiesToSchema(OutputSchema, OutputProperties);
 	AddRequiredListToSchema(OutputSchema, TArray<FString>({ReturnValueName}));
 
@@ -240,7 +229,6 @@ TSharedPtr<FJsonObject> GetFakeToolsetWithMapParamsExpectedSchema()
 	FJsonDomBuilder::FObject OutputSchema = MakeSchema();
 	FJsonDomBuilder::FObject OutputProperties;
 	OutputProperties.Set(ReturnValueName, MakeMapProperty(TEXT("number")));
-	OutputProperties.Set(ErrorValueName, MakeProperty(TEXT("string")));
 	AddPropertiesToSchema(OutputSchema, OutputProperties);
 	AddRequiredListToSchema(OutputSchema, TArray<FString>({ReturnValueName}));
 
@@ -673,16 +661,10 @@ bool FToolsetRegistryToolsetTestExecutePythonToolWithError::RunTest(const FStrin
 	FString Error;
 	UToolsetRegistry::ExecuteTool(
 		"toolset_registry.tests.demo_toolset.ErrorProneToolset", "bad_tool", Input, Result, Error);
-	TestTrue("Tool was called correctly", Error.IsEmpty());
 
-	TSharedPtr<FJsonObject> ResultObject;
-	if (TestTrue("Deserialized tool output correctly",
-				 FJsonSerializer::Deserialize(TJsonReaderFactory<>::Create(Result), ResultObject) &&
-				 ResultObject.IsValid()))
-	{
-		FString ToolError = ResultObject->GetStringField(TEXT("error"));
-		TestFalse("Tool raised internal error", ToolError.IsEmpty());
-	}
+	TestFalse("Tool error was returned", Error.IsEmpty());
+	TestTrue("Error contains exception type", Error.Contains(TEXT("NameError")));
+	TestTrue("Error contains exception message", Error.Contains(TEXT("Bad tool bad")));
 
 	PythonCommand.Command = TEXT(
 		"from toolset_registry.tests.demo_toolset import ErrorProneToolset; "
diff --git a/Engine/Plugins/Experimental/ToolsetRegistry/Source/ToolsetRegistry/Private/ToolsetRegistryToolCall.cpp b/Engine/Plugins/Experimental/ToolsetRegistry/Source/ToolsetRegistry/Private/ToolsetRegistryToolCall.cpp
--- a/Engine/Plugins/Experimental/ToolsetRegistry/Source/ToolsetRegistry/Private/ToolsetRegistryToolCall.cpp
+++ b/Engine/Plugins/Experimental/ToolsetRegistry/Source/ToolsetRegistry/Private/ToolsetRegistryToolCall.cpp
@@ -447,6 +447,12 @@ namespace UE::ToolsetRegistry
 				ReturnJsonObject->SetField(TEXT("returnValue"), ReturnJsonValue);
 			}
 		}
+		else
+		{
+			// If the return value is void, declare the returnValue type as "null".
+			ReturnJsonObject = MakeShared<FJsonObject>();
+			ReturnJsonObject->SetStringField(FString("returnValue"), FString("null"));
+		}
 
 		return MakeValue(ReturnJsonObject);
 	}
diff --git a/Engine/Plugins/Experimental/ToolsetRegistry/Source/ToolsetRegistry/Private/ToolsetRegistryToolset.cpp b/Engine/Plugins/Experimental/ToolsetRegistry/Source/ToolsetRegistry/Private/ToolsetRegistryToolset.cpp
--- a/Engine/Plugins/Experimental/ToolsetRegistry/Source/ToolsetRegistry/Private/ToolsetRegistryToolset.cpp
+++ b/Engine/Plugins/Experimental/ToolsetRegistry/Source/ToolsetRegistry/Private/ToolsetRegistryToolset.cpp
@@ -124,36 +124,6 @@ namespace UE::ToolsetRegistry
 				ToolEntry->SetStringField(TEXT("name"), *NewToolName);
 			}
 
-			// Insert error field into outputSchema, creating it if it doesn't exist.
-			TSharedPtr<FJsonObject> OutputSchemaProperties;
-			TSharedPtr<FJsonObject> OutputSchema;
-			FString OutputSchemaName = TEXT("outputSchema");
-			if (!ToolEntry->HasField(OutputSchemaName))
-			{
-				OutputSchema = MakeShared<FJsonObject>();
-				ToolEntry->SetObjectField(OutputSchemaName, OutputSchema);
-				OutputSchema->SetStringField(TEXT("type"), TEXT("object"));
-				OutputSchemaProperties = MakeShared<FJsonObject>();
-				TSharedPtr<FJsonObject> ReturnValueProperty = MakeShared<FJsonObject>();
-				ReturnValueProperty->SetField(TEXT("type"), MakeShared<FJsonValueNull>());
-				OutputSchemaProperties->SetObjectField(TEXT("returnValue"), ReturnValueProperty);
-				OutputSchema->SetObjectField(TEXT("properties"), OutputSchemaProperties);
-			}
-			else
-			{
-				OutputSchema = ToolEntry->GetObjectField(OutputSchemaName);
-				check(OutputSchema->HasField(TEXT("properties")));
-				OutputSchemaProperties = OutputSchema->GetObjectField(TEXT("properties"));
-			}
-
-			const FString ErrorFieldName = FString(TEXT("error"));
-			if (!OutputSchemaProperties->HasField(ErrorFieldName))
-			{
-				TSharedPtr<FJsonObject> ErrorProperty = MakeShared<FJsonObject>();
-				ErrorProperty->SetStringField(TEXT("type"), TEXT("string"));
-				OutputSchemaProperties->SetObjectField(ErrorFieldName, ErrorProperty);
-			}
-
 			ToolsArray.Add(MakeShared<FJsonValueObject>(ToolEntry));
 		}
 		Schema->SetArrayField(FString(TEXT("tools")), ToolsArray);
@@ -203,57 +173,37 @@ namespace UE::ToolsetRegistry
 		TSharedPtr<FObjectFunctionToolCall> ToolCall = *ToolCallPtr;
 		check(ToolCall);
 
-		TSharedPtr<FJsonObject> ToolCallReturnValue;
 		FObjectFunctionToolCall::FFunctionInputParamsJson JsonParams;
 		JsonParams.Set<FString>(JsonInput);
 
-		// Capture script errors raised by the tool during execution.  Note: script errors raised by tools
-		// running outside of the game thread will not be captured correctly.
+		// Capture script errors raised by the tool during execution. Note: script errors raised by
+		// tools running outside of the game thread will not be captured correctly.
 		FToolCallScriptExceptionHandler ExceptionHandler;
 		ensureAlways(IsInGameThread());
 
 		FObjectFunctionToolCall::FToolCallResult ToolCallResult = ToolCall->Execute(JsonParams);
-
 		if (ToolCallResult.HasError())
 		{
-			FString Error = ToolCallResult.StealError();
-			OutError = FString::Printf(TEXT("Error executing tool %s: %s"), *ToolName, *Error);
+			OutError = ToolCallResult.StealError();
 			return;
 		}
-
-		if (ToolCallResult.HasValue())
-		{
-			ToolCallReturnValue = ToolCallResult.StealValue();
-		}
-		else
+		else if (FString ScriptException = ExceptionHandler.GetException();
+		         !ScriptException.IsEmpty())
 		{
-			OutError = FString::Printf(TEXT("Invalid return value from executing tool %s"), *ToolName);
+			// If the tool raised an error, return that.
+			OutError = ScriptException;
 			return;
 		}
-
-		if (!ToolCallReturnValue.IsValid())
-		{
-			// There was a returned value, but it was null.
-			ToolCallReturnValue = MakeShared<FJsonObject>();
-			ToolCallReturnValue->SetStringField(FString("returnValue"), FString("null"));
-		}
-
-		// If the tool raised an error, return that.
-		if (FString ScriptException = ExceptionHandler.GetException();
-			!ScriptException.IsEmpty())
-		{
-			ToolCallReturnValue->SetStringField(TEXT("error"), MoveTemp(ScriptException));
-		}
-
-		if (!FJsonSerializer::Serialize(
-				ToolCallReturnValue.ToSharedRef(),
+		check(ToolCallResult.HasValue());
+		TSharedPtr<FJsonObject> ResultValue = ToolCallResult.StealValue();
+		if (!ResultValue.IsValid()
+			|| !FJsonSerializer::Serialize(
+				ResultValue.ToSharedRef(),
 				TJsonWriterFactory<TCHAR, TCondensedJsonPrintPolicy<TCHAR>>::Create(&OutResult)))
 		{
 			OutError = FString::Printf(TEXT("Unable to serialize result from executing tool %s"), *ToolName);
 			return;
 		}
-
-		return;
 	}
 
 	TArray<FString> FToolsetLibraryWrapper::GetToolNames()
