Fixed a bug where ObjectTools list properties returned the entire class schema, not just the properties part.

Also added tests for all of ObjectTools.

#Synced-CL 50690832

diff --git a/Engine/Plugins/Experimental/ToolsetRegistry/Source/ToolsetRegistry/Private/ToolsetLibrary.cpp b/Engine/Plugins/Experimental/ToolsetRegistry/Source/ToolsetRegistry/Private/ToolsetLibrary.cpp
--- a/Engine/Plugins/Experimental/ToolsetRegistry/Source/ToolsetRegistry/Private/ToolsetLibrary.cpp
+++ b/Engine/Plugins/Experimental/ToolsetRegistry/Source/ToolsetRegistry/Private/ToolsetLibrary.cpp
@@ -11,10 +11,18 @@
 
 FString UToolsetLibraryTemp::ListClassProperties(UClass* Class)
 {
+	FString JsonProperties;
 	FJsonSchemaPropertyFilter Filter;
 	Filter.CheckFlags |= CPF_BlueprintVisible;
 	Filter.SkipFlags |= CPF_Deprecated;
-	return FJsonSchemaGenerator::UStructToJsonSchemaString(Class, Filter);
+	TSharedPtr<FJsonObject> ClassSchema = FJsonSchemaGenerator::UStructToJsonSchemaObject(Class, Filter);
+	const TSharedPtr<FJsonObject>* ClassProperties = nullptr;
+	if (ClassSchema->TryGetObjectField(TEXT("properties"), ClassProperties))
+	{
+		TSharedRef<TJsonWriter<>> Writer = TJsonWriterFactory<>::Create(&JsonProperties);
+		FJsonSerializer::Serialize(ClassProperties->ToSharedRef(), Writer);
+	}
+	return JsonProperties;
 }
 
 FString UToolsetLibraryTemp::GetObjectProperties(const UObject* Object, const TArray<FName>& PropertyNames)
