[Toolset Registry] Fix issue with results not conforming to output schema when returning non-plain values.

Also fix no-return tool call results to return null object, not "null" as a string.

#rnx

#Synced-CL 50755333

diff --git a/Engine/Plugins/Experimental/ToolsetRegistry/Source/ToolsetRegistry/Private/Tests/ToolsetRegistryToolsetTest.cpp b/Engine/Plugins/Experimental/ToolsetRegistry/Source/ToolsetRegistry/Private/Tests/ToolsetRegistryToolsetTest.cpp
--- a/Engine/Plugins/Experimental/ToolsetRegistry/Source/ToolsetRegistry/Private/Tests/ToolsetRegistryToolsetTest.cpp
+++ b/Engine/Plugins/Experimental/ToolsetRegistry/Source/ToolsetRegistry/Private/Tests/ToolsetRegistryToolsetTest.cpp
@@ -555,7 +555,7 @@ bool FToolsetRegistryToolsetTestExecuteToolWithNoReturn::RunTest(const FString&
 		FakeToolsetName, FString("SimpleNoReturnTool"), FakeToolsetInput, Result, Error);
 	TestTrue("Tool execution with no return value runs without error", Error.IsEmpty());
 
-	FString ExpectedOutput = FString(TEXT(R"({"returnValue":"null"})"));
+	FString ExpectedOutput = FString(TEXT(R"({"returnValue":null})"));
 	TestEqual("Tool execution returns expected output", Result, ExpectedOutput);
 
 	return true;
diff --git a/Engine/Plugins/Experimental/ToolsetRegistry/Source/ToolsetRegistry/Private/ToolsetRegistryToolCall.cpp b/Engine/Plugins/Experimental/ToolsetRegistry/Source/ToolsetRegistry/Private/ToolsetRegistryToolCall.cpp
--- a/Engine/Plugins/Experimental/ToolsetRegistry/Source/ToolsetRegistry/Private/ToolsetRegistryToolCall.cpp
+++ b/Engine/Plugins/Experimental/ToolsetRegistry/Source/ToolsetRegistry/Private/ToolsetRegistryToolCall.cpp
@@ -418,42 +418,41 @@ namespace UE::ToolsetRegistry
 		Function->Invoke(InstanceObject, StackFrame, ReturnValueAddress);
 
 		// Get the return value, if we have one.
-		TSharedPtr<FJsonObject> ReturnJsonObject;
-		if (FProperty* ReturnProperty = Function->GetReturnProperty())
-		{
-			const TSharedPtr<FJsonValue> ReturnJsonValue =
-				FJsonObjectConverter::UPropertyToJsonValue(ReturnProperty, ReturnValueAddress);
-			if (!ReturnJsonValue)
-			{
-				// Caller error. Unsupported return value type. 
-				return MakeError(FString::Printf(TEXT(
-					"Function \"%s\", "
-					"could not convert function return value to Json value. "
-					"Function return value type (%s) may be unsupported."),
-					*Function->GetName(),
-					*ReturnProperty->GetClass()->GetName()));
-			}
+		TSharedPtr<FJsonValue> ReturnJsonValue = Function->GetReturnProperty() ?
+			FJsonObjectConverter::UPropertyToJsonValue(Function->GetReturnProperty(), ReturnValueAddress) :
+			MakeShared<FJsonValueNull>();
 
-			if (ReturnJsonValue->Type == EJson::Object)
-			{
-				// If value is an object, return it.
-				ReturnJsonObject = ReturnJsonValue->AsObject();
-			}
-			else
-			{
-				// If value is not an object, return it inside an object.
-				ReturnJsonObject = MakeShared<FJsonObject>();
-				ReturnJsonObject->SetField(TEXT("returnValue"), ReturnJsonValue);
-			}
+		FString ErrorMessage;
+		TSharedPtr<FJsonObject> ReturnJsonObject = MakeReturnObject(ReturnJsonValue, ErrorMessage);
+
+		if (!ErrorMessage.IsEmpty())
+		{
+			return MakeError(ErrorMessage);
 		}
 		else
 		{
-			// If the return value is void, declare the returnValue type as "null".
-			ReturnJsonObject = MakeShared<FJsonObject>();
-			ReturnJsonObject->SetStringField(FString("returnValue"), FString("null"));
+			return MakeValue(ReturnJsonObject);
+		}
+	}
+
+	TSharedPtr<FJsonObject> FObjectFunctionToolCall::MakeReturnObject(
+		const TSharedPtr<FJsonValue>& InReturnValue, FString& OutErrorMessage) const
+	{
+		// If the value to return is invalid, return an error message instead of the return object.
+		OutErrorMessage.Empty();
+		if (!InReturnValue.IsValid())
+		{
+			OutErrorMessage = FString::Printf(
+				TEXT("Function \"%s\": could not convert function return value to Json value."),
+				*Function->GetName());
+			return nullptr;
 		}
 
-		return MakeValue(ReturnJsonObject);
+		// According to the output schema: all returned values are required to be wrapped in a "returnValue"
+		// object.
+		TSharedPtr<FJsonObject> ReturnJsonObject = MakeShared<FJsonObject>();
+		ReturnJsonObject->SetField(TEXT("returnValue"), InReturnValue);
+		return ReturnJsonObject;
 	}
 
 	TValueOrError<const TSharedRef<FJsonObject>, const FString>
diff --git a/Engine/Plugins/Experimental/ToolsetRegistry/Source/ToolsetRegistry/Public/ToolsetRegistryToolCall.h b/Engine/Plugins/Experimental/ToolsetRegistry/Source/ToolsetRegistry/Public/ToolsetRegistryToolCall.h
--- a/Engine/Plugins/Experimental/ToolsetRegistry/Source/ToolsetRegistry/Public/ToolsetRegistryToolCall.h
+++ b/Engine/Plugins/Experimental/ToolsetRegistry/Source/ToolsetRegistry/Public/ToolsetRegistryToolCall.h
@@ -77,7 +77,10 @@ namespace UE::ToolsetRegistry
 		TValueOrError<const TSharedRef<FJsonObject>, const FString> 
 		BuildValidFunctionInputParamsJsonObject(
 			const TOptional<FFunctionInputParamsJson>& FunctionInputParamsJson) const;
-		
+
+		TSharedPtr<FJsonObject> MakeReturnObject(
+			const TSharedPtr<FJsonValue>& InReturnValue, FString& OutErrorMessage) const;
+
 		TNotNull<UObject*> InstanceObject;
 		TNotNull<UFunction*> Function;
 		TSharedPtr<FJsonObject> FunctionSchemaJsonObject;
