[Toolset Registry] Add tests for C++ tools raising runtime errors

#rnx

#Synced-CL 50825413

diff --git a/Engine/Plugins/Experimental/ToolsetRegistry/Source/ToolsetRegistry/Private/Tests/ToolsetRegistryToolsetTest.cpp b/Engine/Plugins/Experimental/ToolsetRegistry/Source/ToolsetRegistry/Private/Tests/ToolsetRegistryToolsetTest.cpp
--- a/Engine/Plugins/Experimental/ToolsetRegistry/Source/ToolsetRegistry/Private/Tests/ToolsetRegistryToolsetTest.cpp
+++ b/Engine/Plugins/Experimental/ToolsetRegistry/Source/ToolsetRegistry/Private/Tests/ToolsetRegistryToolsetTest.cpp
@@ -633,6 +633,31 @@ bool FToolsetRegistryToolsetTestValidateToolMethods::RunTest(const FString& Unus
 	return true;
 }
 
+// Invoke a C++ tool that produces errors, not from bad input.
+IMPLEMENT_SIMPLE_AUTOMATION_TEST(
+	FToolsetRegistryToolsetTestExecuteCppToolWithError,
+	"AI.ToolsetRegistry.Toolset.ExecuteCppToolWithError",
+	ToolsetRegistryTest::Flags)
+
+bool FToolsetRegistryToolsetTestExecuteCppToolWithError::RunTest(const FString& UnusedParameters)
+{
+	ToolsetRegistryContainer Container(UFakeErrorProneToolset::StaticClass());
+	TSharedPtr<UE::ToolsetRegistry::FToolsetLibraryWrapper> Toolset = Container.ToolsetHandler();
+	if (!TestTrue("Valid toolset handler returned", Toolset.IsValid()))
+	{
+		return false;
+	}
+
+	FString Result;
+	FString Error;
+	UToolsetRegistry::ExecuteTool(Toolset->GetToolsetName(), TEXT("BuggyTool"), TEXT(""), Result, Error);
+	TestFalse("Tool execution with buggy tool returns with error", Error.IsEmpty());
+	TestEqual("Tool execution with buggy tool returns correct raised exception",
+			  Error, UFakeErrorProneToolset::ErrorMessage);
+
+	return true;
+}
+
 // Invoke python tool that produces errors.
 IMPLEMENT_SIMPLE_AUTOMATION_TEST(
 	FToolsetRegistryToolsetTestExecutePythonToolWithError,
diff --git a/Engine/Plugins/Experimental/ToolsetRegistry/Source/ToolsetRegistry/Private/Tests/ToolsetRegistryToolsetTest.h b/Engine/Plugins/Experimental/ToolsetRegistry/Source/ToolsetRegistry/Private/Tests/ToolsetRegistryToolsetTest.h
--- a/Engine/Plugins/Experimental/ToolsetRegistry/Source/ToolsetRegistry/Private/Tests/ToolsetRegistryToolsetTest.h
+++ b/Engine/Plugins/Experimental/ToolsetRegistry/Source/ToolsetRegistry/Private/Tests/ToolsetRegistryToolsetTest.h
@@ -5,6 +5,8 @@
 #include "CoreMinimal.h"
 
 #include "Dom/JsonObject.h"
+#include "Kismet/KismetSystemLibrary.h"
+
 #include "BPToolsetRegistry.h"
 
 #include "ToolsetRegistryToolsetTest.generated.h"
@@ -149,3 +151,21 @@ private:
 	UFUNCTION()
 	static void SimplePrivateMethod() {}
 };
+
+/** Fake Toolset that has error-prone method. */
+UCLASS(Blueprintable)
+class UFakeErrorProneToolset : public UToolsetDefinition
+{
+	GENERATED_BODY()
+
+public:
+
+	static inline FString ErrorMessage = TEXT("Bug in this tool");
+
+	/** Method that raises errors. */
+	UFUNCTION(BlueprintCallable, Category=Fake)
+	static void BuggyTool()
+	{
+		UKismetSystemLibrary::RaiseScriptError(ErrorMessage);
+	}
+};
