[Toolset Registry] Capture all exceptions from script runtime when running tools and treat as errors.

#rnx

Reason: Bughawk internal errors due to not finding python files in CL in its cache
#Synced-CL 50825413

diff --git a/Engine/Plugins/Experimental/ToolsetRegistry/Source/ToolsetRegistry/Private/ToolsetRegistryToolset.cpp b/Engine/Plugins/Experimental/ToolsetRegistry/Source/ToolsetRegistry/Private/ToolsetRegistryToolset.cpp
--- a/Engine/Plugins/Experimental/ToolsetRegistry/Source/ToolsetRegistry/Private/ToolsetRegistryToolset.cpp
+++ b/Engine/Plugins/Experimental/ToolsetRegistry/Source/ToolsetRegistry/Private/ToolsetRegistryToolset.cpp
@@ -7,10 +7,12 @@
 #include "Dom/JsonObject.h"
 #include "JsonSchema/JsonSchemaGenerator.h"
 #include "Kismet/BlueprintFunctionLibrary.h"
+#include "Misc/CoreMisc.h"
+#include "Misc/Optional.h"
 #include "Policies/CondensedJsonPrintPolicy.h"
 #include "Serialization/JsonWriter.h"
 #include "Serialization/JsonSerializer.h"
-#include "Blueprint/BlueprintExceptionInfo.h" 
+#include "Blueprint/BlueprintExceptionInfo.h"
 #include "ToolsetRegistryDelegateHandle.h"
 
 #include "ToolsetRegistryToolCall.h"
@@ -29,17 +31,29 @@ namespace UE::ToolsetRegistry
 					FDelegateHandleRaii::Create(
 						FBlueprintCoreDelegates::OnScriptException,
 						FBlueprintCoreDelegates::OnScriptException.AddRaw(
-							this, &FToolCallScriptExceptionHandler::OnScriptException)))
+							this, &FToolCallScriptExceptionHandler::OnScriptException))),
+				  RuntimeScriptExceptionHandle(
+					[this](ELogVerbosity::Type Verbosity, const TCHAR* ExceptionMessage, const TCHAR* StackMessage)
+					{
+						AddToExceptionMessage(FString(ExceptionMessage));
+
+						// Also log to console to make more obvious to tool developers.
+						UE_LOG(LogToolsetRegistry, Warning, TEXT("Error running tool:\n%s"), ExceptionMessage);
+					})
 			{}
+
 			~FToolCallScriptExceptionHandler() = default;
 
-			// Track script exceptions and copy into local storage.
+			// Track script exceptions raised by UKismetSystemLibrary::RaiseScriptException or
+			// FBlueprintCoreDelegates::ThrowScriptException via FBlueprintCoreDelegates::OnScriptException
+			// and copy into local storage.
 			void OnScriptException(
 				const UObject* ActiveObject, const FFrame& StackFrame, const FBlueprintExceptionInfo& Info)
 			{
-				if (Info.GetType() == EBlueprintExceptionType::UserRaisedError)
+				TOptional<FString> ErrorMessage = GetBlueprintExceptionMessage(Info);
+				if (ErrorMessage.IsSet())
 				{
-					ScriptException = Info.GetDescription().ToString();
+					AddToExceptionMessage(ErrorMessage.GetValue());
 				}
 			}
 
@@ -49,7 +63,77 @@ namespace UE::ToolsetRegistry
 			}
 
 		  private:
+			TOptional<FString> GetBlueprintExceptionMessage(const FBlueprintExceptionInfo& Info)
+			{
+				TOptional<FString> ExceptionMessage;
+				switch (Info.GetType())
+				{
+					// Ignore breakpoints & tracepoints.
+					case EBlueprintExceptionType::Breakpoint:
+						// Ignore and fall through
+					case EBlueprintExceptionType::Tracepoint:
+						// Ignore and fall through
+					case EBlueprintExceptionType::WireTracepoint:
+						// Return unset message, to be ignored.
+						break;
+					case EBlueprintExceptionType::AccessViolation:
+						ExceptionMessage = Info.GetDescription().IsEmpty() ?
+							TEXT("AccessViolation") :
+							Info.GetDescription().ToString();
+						break;
+					case EBlueprintExceptionType::InfiniteLoop:
+						ExceptionMessage = Info.GetDescription().IsEmpty() ?
+							TEXT("InfiniteLoop") :
+							Info.GetDescription().ToString();
+						break;
+					case EBlueprintExceptionType::NonFatalError:
+						ExceptionMessage = Info.GetDescription().IsEmpty() ?
+							TEXT("NonFatalError") :
+							Info.GetDescription().ToString();
+						break;
+					case EBlueprintExceptionType::FatalError:
+						ExceptionMessage = Info.GetDescription().IsEmpty() ?
+							TEXT("FatalError") :
+							Info.GetDescription().ToString();
+						break;
+					case EBlueprintExceptionType::AbortExecution:
+						ExceptionMessage = Info.GetDescription().IsEmpty() ?
+							TEXT("AbortExecution") :
+							Info.GetDescription().ToString();
+						break;
+					case EBlueprintExceptionType::UserRaisedError:
+						ExceptionMessage = Info.GetDescription().IsEmpty() ?
+							TEXT("UserRaisedError") :
+							Info.GetDescription().ToString();
+						break;
+					default:
+						checkNoEntry();
+						break;
+				}
+				return ExceptionMessage;
+			}
+
+			void AddToExceptionMessage(const FString& NewMessage)
+			{
+				if (NewMessage.IsEmpty())
+				{
+					return;
+				}
+				else if (ScriptException.IsEmpty())
+				{
+					ScriptException = NewMessage;
+				}
+				else
+				{
+					ScriptException += FString::Printf(TEXT("\n%s"), *NewMessage);
+				}
+			}
+
+			// Handler for FBlueprintCoreDelegates::OnScriptException
 			FDelegateHandleRaii ScriptExceptionHandle;
+			// Handler for FScriptExceptionHandler
+			FScopedScriptExceptionHandler RuntimeScriptExceptionHandle;
+			// Raised script exception message.
 			FString ScriptException;
 		};
 	}
