[Toolset Registry] Shorten error message on python tool exception

#rnx

Bughawk disabled due to spurious errors trying to access Python files not in its cache, no issues found on the C++ files

#Synced-CL 50825413

diff --git a/Engine/Plugins/Experimental/ToolsetRegistry/Source/ToolsetRegistry/Private/Tests/ToolsetRegistryToolsetTest.cpp b/Engine/Plugins/Experimental/ToolsetRegistry/Source/ToolsetRegistry/Private/Tests/ToolsetRegistryToolsetTest.cpp
--- a/Engine/Plugins/Experimental/ToolsetRegistry/Source/ToolsetRegistry/Private/Tests/ToolsetRegistryToolsetTest.cpp
+++ b/Engine/Plugins/Experimental/ToolsetRegistry/Source/ToolsetRegistry/Private/Tests/ToolsetRegistryToolsetTest.cpp
@@ -684,12 +684,12 @@ bool FToolsetRegistryToolsetTestExecutePythonToolWithError::RunTest(const FStrin
 	FString Input = TEXT(R"({"msg":"Hi"})");
 	FString Result;
 	FString Error;
+	FString ExpectedError = TEXT("Bad tool bad");
 	UToolsetRegistry::ExecuteTool(
 		"toolset_registry.tests.demo_toolset.ErrorProneToolset", "bad_tool", Input, Result, Error);
 
 	TestFalse("Tool error was returned", Error.IsEmpty());
-	TestTrue("Error contains exception type", Error.Contains(TEXT("NameError")));
-	TestTrue("Error contains exception message", Error.Contains(TEXT("Bad tool bad")));
+	TestEqual("Error is expected exception message", Error, ExpectedError);
 
 	PythonCommand.Command = TEXT(
 		"from toolset_registry.tests.demo_toolset import ErrorProneToolset; "
