[AI Assistant] Fix overly restrictive definition of FObjectFunctionToolCall::FToolCallResult

TValueOrError is copied by value the fields should not be const.

#Synced-CL 50892692

diff --git a/Engine/Plugins/Experimental/ToolsetRegistry/Source/ToolsetRegistry/Public/ToolsetRegistryToolCall.h b/Engine/Plugins/Experimental/ToolsetRegistry/Source/ToolsetRegistry/Public/ToolsetRegistryToolCall.h
--- a/Engine/Plugins/Experimental/ToolsetRegistry/Source/ToolsetRegistry/Public/ToolsetRegistryToolCall.h
+++ b/Engine/Plugins/Experimental/ToolsetRegistry/Source/ToolsetRegistry/Public/ToolsetRegistryToolCall.h
@@ -29,7 +29,7 @@ namespace UE::ToolsetRegistry
 		 * Tool call output result.
 		 * Either a JSON object or a string describing an error.
 		 */
-		using FToolCallResult = TValueOrError<const TSharedPtr<FJsonObject>, const FString>;
+		using FToolCallResult = TValueOrError<TSharedPtr<FJsonObject>, FString>;
 		
 		/**
 		 * Initializes use of this class.
