Added library for Toolset specific JSON transformations.

This library allows us to consistently apply overrides to JSON schema/data generation that is available in JsonUtilities, like representing all UObject/UClass references as a special type.

#Synced-CL 50884079

diff --git a/Engine/Plugins/Experimental/AIAssistant/Source/AIAssistant/Private/AIAssistantToolset.cpp b/Engine/Plugins/Experimental/AIAssistant/Source/AIAssistant/Private/AIAssistantToolset.cpp
--- a/Engine/Plugins/Experimental/AIAssistant/Source/AIAssistant/Private/AIAssistantToolset.cpp
+++ b/Engine/Plugins/Experimental/AIAssistant/Source/AIAssistant/Private/AIAssistantToolset.cpp
@@ -6,7 +6,6 @@
 #include "EdGraphUtilities.h"
 #include "IMaterialEditor.h"
 #include "IPythonScriptPlugin.h"
-#include "JsonSchema/JsonSchemaGenerator.h"
 #include "MaterialGraph/MaterialGraph.h"
 #include "Materials/Material.h"
 #include "UObject/ObjectPtr.h"
diff --git a/Engine/Plugins/Experimental/ToolsetRegistry/Source/ToolsetRegistry/Private/Tests/ToolsetJsonTest.cpp b/Engine/Plugins/Experimental/ToolsetRegistry/Source/ToolsetRegistry/Private/Tests/ToolsetJsonTest.cpp
new file mode 100644
--- /dev/null
+++ b/Engine/Plugins/Experimental/ToolsetRegistry/Source/ToolsetRegistry/Private/Tests/ToolsetJsonTest.cpp
@@ -0,0 +1,370 @@
+// Copyright Epic Games, Inc. All Rights Reserved.
+
+#include "ToolsetJsonTest.h"
+#include "CoreMinimal.h"
+
+#include "JsonObjectConverter.h"
+#include "Misc/AutomationTest.h"
+#include "ToolsetJson.h"
+
+
+#if WITH_DEV_AUTOMATION_TESTS
+
+void TestObjectRefSchema(FAutomationTestBase& Test, const TSharedPtr<FJsonObject>& Schema, UClass* Class, bool IsClass)
+{
+	Test.TestTrue("Schema", Schema.IsValid());
+	if (!Schema.IsValid())
+	{
+		return;
+	}
+
+	FString Type;
+	Test.TestTrue("Type", Schema->TryGetStringField(TEXT("type"), Type));
+	Test.TestEqual("Type", Type, FString(TEXT("object")));
+
+	if (Class)
+	{
+		FString Title;
+		Test.TestTrue("Title", Schema->TryGetStringField(TEXT("title"), Title));
+		FString ClassPrefix(IsClass ? TEXT("Class@") : TEXT(""));
+		Test.TestEqual("Title", Title, ClassPrefix + Class->GetPathName());
+	}	
+	
+	const TSharedPtr<FJsonObject>* Properties = nullptr;
+	Test.TestTrue("Properties", Schema->TryGetObjectField(TEXT("properties"), Properties));
+	if (Properties)
+	{
+		const TSharedPtr<FJsonObject>* RefPath = nullptr;
+		Test.TestTrue("RefPath", Properties->Get()->TryGetObjectField(TEXT("refPath"), RefPath));
+		if (RefPath)
+		{
+			Test.TestTrue("RefPath", RefPath->Get()->TryGetStringField(TEXT("type"), Type));
+			Test.TestEqual("RefPath", Type, FString(TEXT("string")));
+		}
+	}
+}
+
+void TestObjectRefData(
+	FAutomationTestBase& Test, const TSharedPtr<FJsonObject>& Schema, const FString& Expected)
+{
+	Test.TestTrue("Schema", Schema.IsValid());
+	if (!Schema.IsValid())
+	{
+		return;
+	}
+	
+	FString RefPath;
+	Test.TestTrue("Unpack", Schema->TryGetStringField(TEXT("refPath"), RefPath));
+	Test.TestEqual("RefPath", RefPath, Expected);
+}
+
+void TestObjectRefDataInStruct(
+	FAutomationTestBase& Test, const TSharedPtr<FJsonObject>& Schema, const FString& PropertyName,
+	const FString& Expected)
+{
+	Test.TestTrue("Schema", Schema.IsValid());
+	if (!Schema.IsValid())
+	{
+		return;
+	}
+
+	const TSharedPtr<FJsonObject>* Property = nullptr;
+	Test.TestTrue("Unpack", Schema->TryGetObjectField(PropertyName, Property));
+	if (Property)
+	{
+		TestObjectRefData(Test, *Property, Expected);
+	}
+}
+
+BEGIN_DEFINE_SPEC(FToolsetJsonSpec, "AI.ToolsetRegistry.ToolsetJsonSpec",
+	EAutomationTestFlags::EditorContext | EAutomationTestFlags::EngineFilter)	
+END_DEFINE_SPEC(FToolsetJsonSpec)
+
+void FToolsetJsonSpec::Define()
+{
+	It("Converts UObject* property to reference schema", [this]()
+	{
+		UStruct* Struct = FToolsetJsonTest::StaticStruct();
+		FProperty* Prop = Struct->FindPropertyByName("TestObject");
+		const TSharedPtr<FJsonObject> Schema = FToolsetJson::PropertyToJsonSchema(Prop);
+		TestObjectRefSchema(*this, Schema, UToolsetJsonTestObject::StaticClass(), false);
+	});
+
+	It("Converts UObject* property to and from JSON", [this]()
+	{
+		FToolsetJsonTest TestIn;
+		TObjectPtr<UToolsetJsonTestObject> TestOther = NewObject<UToolsetJsonTestObject>();
+		TestIn.TestObject = TestOther;
+
+		UStruct* Struct = FToolsetJsonTest::StaticStruct();
+		FProperty* Property = Struct->FindPropertyByName("TestObject");
+
+		TSharedPtr<FJsonValue> OutJson = FToolsetJson::PropertyToJsonData(Property, &TestIn.TestObject);
+		TestObjectRefData(*this, OutJson->AsObject(), TestOther->GetPathName());
+
+		FToolsetJsonTest TestOut;
+		TestTrue("Output", FToolsetJson::JsonDataToProperty(OutJson, Property, &TestOut.TestObject));
+		TestEqual("RefCheck", TestIn.TestObject, TestOut.TestObject);
+	});
+
+	It("Converts SoftPath property to reference schema", [this]()
+	{
+		UStruct* Struct = FToolsetJsonTest::StaticStruct();
+		FProperty* Prop = Struct->FindPropertyByName("TestObjectPath");
+		const TSharedPtr<FJsonObject> Schema = FToolsetJson::PropertyToJsonSchema(Prop);
+		TestObjectRefSchema(*this, Schema, UObject::StaticClass(), false);
+	});
+
+	It("Converts SoftPath property to and from JSON", [this]()
+	{
+		FToolsetJsonTest TestIn;
+		TObjectPtr<UToolsetJsonTestObject> TestOther = NewObject<UToolsetJsonTestObject>();
+		TestIn.TestObjectPath = TestOther;
+
+		UStruct* Struct = FToolsetJsonTest::StaticStruct();
+		FProperty* Property = Struct->FindPropertyByName("TestObjectPath");
+
+		TSharedPtr<FJsonValue> OutJson = FToolsetJson::PropertyToJsonData(Property, &TestIn.TestObjectPath);
+		TestObjectRefData(*this, OutJson->AsObject(), TestOther->GetPathName());
+
+		FToolsetJsonTest TestOut;
+		TestTrue("Output", FToolsetJson::JsonDataToProperty(OutJson, Property, &TestOut.TestObjectPath));
+		TestEqual("RefCheck", TestIn.TestObjectPath, TestOut.TestObjectPath);
+	});
+
+	It("Converts SoftPtr property to reference schema", [this]()
+	{
+		UStruct* Struct = FToolsetJsonTest::StaticStruct();
+		FProperty* Prop = Struct->FindPropertyByName("TestObjectSoft");
+		const TSharedPtr<FJsonObject> Schema = FToolsetJson::PropertyToJsonSchema(Prop);
+		TestObjectRefSchema(*this, Schema, UToolsetJsonTestObject::StaticClass(), false);
+	});
+
+	It("Converts SoftPtr property to and from JSON", [this]()
+	{
+		FToolsetJsonTest TestIn;
+		TObjectPtr<UToolsetJsonTestObject> TestOther = NewObject<UToolsetJsonTestObject>();
+		TestIn.TestObjectSoft = TestOther;
+
+		UStruct* Struct = FToolsetJsonTest::StaticStruct();
+		FProperty* Property = Struct->FindPropertyByName("TestObjectSoft");
+
+		TSharedPtr<FJsonValue> OutJson = FToolsetJson::PropertyToJsonData(Property, &TestIn.TestObjectSoft);
+		TestObjectRefData(*this, OutJson->AsObject(), TestOther->GetPathName());
+
+		FToolsetJsonTest TestOut;
+		TestTrue("Output", FToolsetJson::JsonDataToProperty(OutJson, Property, &TestOut.TestObjectSoft));
+		TestEqual("RefCheck", TestIn.TestObjectSoft, TestOut.TestObjectSoft);
+	});
+
+	It("Converts WeakPtr property to reference schema", [this]()
+	{
+		UStruct* Struct = FToolsetJsonTest::StaticStruct();
+		FProperty* Prop = Struct->FindPropertyByName("TestObjectWeak");
+		const TSharedPtr<FJsonObject> Schema = FToolsetJson::PropertyToJsonSchema(Prop);
+		TestObjectRefSchema(*this, Schema, UToolsetJsonTestObject::StaticClass(), false);
+	});
+
+	It("Converts WeakPtr property to and from JSON", [this]()
+	{
+		FToolsetJsonTest TestIn;
+		TObjectPtr<UToolsetJsonTestObject> TestOther = NewObject<UToolsetJsonTestObject>();
+		TestIn.TestObjectWeak = TestOther;
+
+		UStruct* Struct = FToolsetJsonTest::StaticStruct();
+		FProperty* Property = Struct->FindPropertyByName("TestObjectWeak");
+
+		TSharedPtr<FJsonValue> OutJson = FToolsetJson::PropertyToJsonData(Property, &TestIn.TestObjectWeak);
+		TestObjectRefData(*this, OutJson->AsObject(), TestOther->GetPathName());
+
+		FToolsetJsonTest TestOut;
+		TestTrue("Output", FToolsetJson::JsonDataToProperty(OutJson, Property, &TestOut.TestObjectWeak));
+		TestEqual("RefCheck", TestIn.TestObjectWeak, TestOut.TestObjectWeak);
+	});
+
+	It("Converts UClass* property to reference schema", [this]()
+	{
+		UStruct* Struct = FToolsetJsonTest::StaticStruct();
+		FProperty* Prop = Struct->FindPropertyByName("TestClass");
+		const TSharedPtr<FJsonObject> Schema = FToolsetJson::PropertyToJsonSchema(Prop);
+		TestObjectRefSchema(*this, Schema, UToolsetJsonTestObject::StaticClass(), true);
+	});
+
+	It("Converts UClass* property to and from JSON", [this]()
+	{
+		FToolsetJsonTest TestIn;
+		TestIn.TestClass = UToolsetJsonTestObject::StaticClass();
+
+		UStruct* Struct = FToolsetJsonTest::StaticStruct();
+		FProperty* Property = Struct->FindPropertyByName("TestClass");
+
+		TSharedPtr<FJsonValue> OutJson = FToolsetJson::PropertyToJsonData(Property, &TestIn.TestClass);
+		TestObjectRefData(*this, OutJson->AsObject(), UToolsetJsonTestObject::StaticClass()->GetPathName());
+
+		FToolsetJsonTest TestOut;
+		TestTrue("Output", FToolsetJson::JsonDataToProperty(OutJson, Property, &TestOut.TestClass));
+		TestEqual("RefCheck", TestIn.TestClass, TestOut.TestClass);
+	});
+
+	It("Converts SoftClassPath property to reference schema", [this]()
+	{
+		UStruct* Struct = FToolsetJsonTest::StaticStruct();
+		FProperty* Prop = Struct->FindPropertyByName("TestClassPath");
+		const TSharedPtr<FJsonObject> Schema = FToolsetJson::PropertyToJsonSchema(Prop);
+		TestObjectRefSchema(*this, Schema, UObject::StaticClass(), true);
+	});
+
+	It("Converts SoftClassPath property to and from JSON", [this]()
+	{
+		FToolsetJsonTest TestIn;
+		TestIn.TestClassPath = UToolsetJsonTestObject::StaticClass();
+
+		UStruct* Struct = FToolsetJsonTest::StaticStruct();
+		FProperty* Property = Struct->FindPropertyByName("TestClassPath");
+
+		TSharedPtr<FJsonValue> OutJson = FToolsetJson::PropertyToJsonData(Property, &TestIn.TestClassPath);
+		TestObjectRefData(*this, OutJson->AsObject(), UToolsetJsonTestObject::StaticClass()->GetPathName());
+
+		FToolsetJsonTest TestOut;
+		TestTrue("Output", FToolsetJson::JsonDataToProperty(OutJson, Property, &TestOut.TestClassPath));
+		TestEqual("RefCheck", TestIn.TestClassPath, TestOut.TestClassPath);
+	});
+
+	It("Converts struct to reference schema", [this]()
+	{		
+		TSharedPtr<FJsonObject> Schema = FToolsetJson::StructToJsonSchema(FToolsetJsonTest::StaticStruct());
+		TestTrue("Convert", Schema.IsValid());
+		if (!Schema.IsValid())
+		{
+			return;
+		}
+		const TSharedPtr<FJsonObject>* Properties = nullptr;
+		TestTrue("Properties", Schema->TryGetObjectField(TEXT("properties"), Properties));
+		if (Properties)
+		{
+			for (TFieldIterator<FProperty> It(FToolsetJsonTest::StaticStruct()); It; ++It)
+			{
+				FProperty* Property = *It;
+				FString JsonPropertyName = FJsonObjectConverter::StandardizeCase(Property->GetName());
+				const TSharedPtr<FJsonObject>* TestProperty = nullptr;
+				TestTrue("Object", Properties->Get()->TryGetObjectField(JsonPropertyName, TestProperty));
+				if (TestProperty)
+				{
+					TestObjectRefSchema(*this, *TestProperty, nullptr, false);
+				}
+			}
+		}
+	});
+
+	It("Converts struct to and from JSON.", [this]()
+	{
+		FToolsetJsonTest TestIn;
+		TObjectPtr<UToolsetJsonTestObject> TestOther = NewObject<UToolsetJsonTestObject>();
+		TestIn.TestObject = TestOther;
+		TestIn.TestObjectPath = TestOther;
+		TestIn.TestObjectSoft = TestOther;
+		TestIn.TestObjectWeak = TestOther;
+		TestIn.TestClass = UToolsetJsonTestObject::StaticClass();
+		TestIn.TestClassPath = UToolsetJsonTestObject::StaticClass();
+
+		TSharedPtr<FJsonObject> JsonObject = FToolsetJson::StructToJsonData(
+			FToolsetJsonTest::StaticStruct(), &TestIn);
+		for (TFieldIterator<FProperty> It(FToolsetJsonTest::StaticStruct()); It; ++It)
+		{
+			FProperty* Property = *It;
+			FString JsonPropertyName = FJsonObjectConverter::StandardizeCase(Property->GetName());
+			if (JsonPropertyName == "testClass" || JsonPropertyName == "testClassPath")
+			{
+				TestObjectRefDataInStruct(
+					*this, JsonObject, JsonPropertyName, UToolsetJsonTestObject::StaticClass()->GetPathName());
+			}
+			else
+			{
+				TestObjectRefDataInStruct(
+					*this, JsonObject, JsonPropertyName, TestOther->GetPathName());
+			}
+		}		
+
+		FToolsetJsonTest TestOut;
+		TestTrue("FromJson", FToolsetJson::JsonDataToStruct(
+			JsonObject.ToSharedRef(), FToolsetJsonTest::StaticStruct(), &TestOut));
+		TestEqual("RefCheck", TestIn.TestObject, TestOut.TestObject);
+		TestEqual("RefCheck", TestIn.TestObjectPath, TestOut.TestObjectPath);
+		TestEqual("RefCheck", TestIn.TestObjectSoft, TestOut.TestObjectSoft);
+		TestEqual("RefCheck", TestIn.TestObjectWeak, TestOut.TestObjectWeak);
+		TestEqual("RefCheck", TestIn.TestClass, TestOut.TestClass);
+		TestEqual("RefCheck", TestIn.TestClassPath, TestOut.TestClassPath);
+	});
+
+	It("Imports from old string style path", [this]()
+	{
+		TObjectPtr<UToolsetJsonTestObject> TestOther = NewObject<UToolsetJsonTestObject>();
+		UStruct* Struct = FToolsetJsonTest::StaticStruct();
+		FProperty* Property = Struct->FindPropertyByName("TestObject");
+
+		FToolsetJsonTest TestOut;
+		TSharedPtr<FJsonValue> OutJson = MakeShared<TJsonValueString<TCHAR>>(TestOther->GetPathName());
+		TestTrue("Output", FToolsetJson::JsonDataToProperty(OutJson, Property, &TestOut.TestObject));
+		TestEqual("RefCheck", TestOut.TestObject, TestOther);
+	});
+
+	It("Filters user visible properties schema", [this]()
+	{
+		TSharedPtr<FJsonObject> ClassSchema = FToolsetJson::StructToJsonSchema(
+			UToolsetJsonTestObject::StaticClass(), true);
+		TestTrue("TestData", ClassSchema.IsValid());
+		if (ClassSchema)
+		{
+			const TSharedPtr<FJsonObject>* Properties = nullptr;
+			TestTrue("Properties", ClassSchema->TryGetObjectField(TEXT("properties"), Properties));
+			if (Properties)
+			{
+				TestTrue("Visible", Properties->Get()->HasField(TEXT("visibleProperty")));
+				TestFalse("Hidden", Properties->Get()->HasField(TEXT("hiddenProperty")));
+			}
+		}
+	});
+
+	It("Filters user visible properties data", [this]()
+	{
+		TSharedPtr<FJsonObject> TestData = FToolsetJson::StructToJsonData(
+			UToolsetJsonTestObject::StaticClass(), NewObject<UToolsetJsonTestObject>(), true);
+		TestTrue("TestData", TestData.IsValid());
+		if (TestData)
+		{
+			TestTrue("Visible", TestData->HasField(TEXT("visibleProperty")));
+			TestFalse("Hidden", TestData->HasField(TEXT("hiddenProperty")));
+		}
+	});
+
+	It("Can read defaults", [this]()
+	{
+		UFunction* FunctionPtr = UToolsetJsonTestObject::StaticClass()->FindFunctionByName("TestDefault");
+		const TSharedPtr<FJsonObject> Schema = FToolsetJson::StructToJsonSchema(FunctionPtr);
+		TestTrue("Schema", Schema.IsValid());
+		if (!Schema.IsValid())
+		{
+			return;
+		}
+		const TSharedPtr<FJsonObject>* InputSchema = nullptr;
+		TestTrue("Input", Schema->TryGetObjectField(TEXT("inputSchema"), InputSchema));
+		if (InputSchema)
+		{
+			TestFalse("NoRequired", InputSchema->Get()->HasField(TEXT("required")));
+			const TSharedPtr<FJsonObject>* Properties = nullptr;
+			TestTrue("Properties", InputSchema->Get()->TryGetObjectField(TEXT("properties"), Properties));
+			if (Properties)
+			{
+				const TSharedPtr<FJsonObject>* TestObject = nullptr;
+				TestTrue("Input", Properties->Get()->TryGetObjectField(TEXT("testObject"), TestObject));
+				if (TestObject)
+				{
+					TestTrue("Default", TestObject->Get()->GetField(TEXT("default"), EJson::Null)->IsNull());
+				}
+			}
+		}
+	});
+
+}
+
+#endif
diff --git a/Engine/Plugins/Experimental/ToolsetRegistry/Source/ToolsetRegistry/Private/Tests/ToolsetJsonTest.h b/Engine/Plugins/Experimental/ToolsetRegistry/Source/ToolsetRegistry/Private/Tests/ToolsetJsonTest.h
new file mode 100644
--- /dev/null
+++ b/Engine/Plugins/Experimental/ToolsetRegistry/Source/ToolsetRegistry/Private/Tests/ToolsetJsonTest.h
@@ -0,0 +1,49 @@
+// Copyright Epic Games, Inc. All Rights Reserved.
+
+#pragma once
+
+#include "CoreMinimal.h"
+#include "Templates/SharedPointer.h"
+#include "Templates/SubclassOf.h"
+
+#include "ToolsetJsonTest.generated.h"
+
+
+UCLASS(BlueprintType)
+class UToolsetJsonTestObject : public UObject
+{
+	GENERATED_BODY()
+public:
+	UPROPERTY(BlueprintReadWrite, Category="ToolsetTest")
+	bool VisibleProperty;
+
+	UPROPERTY()
+	bool HiddenProperty;
+
+	UFUNCTION(BlueprintCallable, Category = "Main")
+	void TestDefault(UObject* TestObject = nullptr) {}
+};
+
+USTRUCT(BlueprintType)
+struct FToolsetJsonTest
+{
+	GENERATED_BODY()
+public:
+	UPROPERTY()
+	TObjectPtr<UToolsetJsonTestObject> TestObject;
+
+	UPROPERTY()
+	FSoftObjectPath TestObjectPath;
+
+	UPROPERTY()
+	TSoftObjectPtr<UToolsetJsonTestObject> TestObjectSoft;
+
+	UPROPERTY()
+	TWeakObjectPtr<UToolsetJsonTestObject> TestObjectWeak;
+
+	UPROPERTY()
+	TSubclassOf<UToolsetJsonTestObject> TestClass;
+
+	UPROPERTY()
+	FSoftClassPath TestClassPath;
+};
diff --git a/Engine/Plugins/Experimental/ToolsetRegistry/Source/ToolsetRegistry/Private/Tests/ToolsetRegistryToolCallTest.cpp b/Engine/Plugins/Experimental/ToolsetRegistry/Source/ToolsetRegistry/Private/Tests/ToolsetRegistryToolCallTest.cpp
--- a/Engine/Plugins/Experimental/ToolsetRegistry/Source/ToolsetRegistry/Private/Tests/ToolsetRegistryToolCallTest.cpp
+++ b/Engine/Plugins/Experimental/ToolsetRegistry/Source/ToolsetRegistry/Private/Tests/ToolsetRegistryToolCallTest.cpp
@@ -4,9 +4,10 @@
 
 #include "Misc/AutomationTest.h"
 #include "UObject/Package.h"
+#include "Serialization/JsonSerializer.h"
+#include "Serialization/JsonWriter.h"
 
 #include "Tests/ToolsetRegistryTestFlags.h"
-#include "JsonSchema/JsonSchemaGenerator.h"
 
 #include "ToolsetRegistryToolCall.h"
 #include "ToolsetRegistryModule.h"
diff --git a/Engine/Plugins/Experimental/ToolsetRegistry/Source/ToolsetRegistry/Private/ToolsetJson.cpp b/Engine/Plugins/Experimental/ToolsetRegistry/Source/ToolsetRegistry/Private/ToolsetJson.cpp
new file mode 100644
--- /dev/null
+++ b/Engine/Plugins/Experimental/ToolsetRegistry/Source/ToolsetRegistry/Private/ToolsetJson.cpp
@@ -0,0 +1,273 @@
+// Copyright Epic Games, Inc. All Rights Reserved.
+
+#include "ToolsetJson.h"
+
+#include "Editor.h"
+#include "JsonObjectConverter.h"
+#include "JsonObjectStructInterface.h"
+#include "JsonSchema/JsonSchemaGenerator.h"
+
+EPropertyFlags GetCheckFlags(bool UserVisiblePropertiesOnly)
+{
+	EPropertyFlags CheckFlags = CPF_None;
+	if (UserVisiblePropertiesOnly)
+	{
+		CheckFlags = CPF_BlueprintVisible | CPF_BlueprintAssignable;
+		if (GEditor)
+		{
+			CheckFlags |= CPF_Edit;
+		}
+	}
+	return CheckFlags;
+}
+
+bool IsReferenceLikeProperty(TNotNull<const FProperty*> Property)
+{
+	if (CastField<FObjectProperty>(Property) && !Property->HasAnyPropertyFlags(CPF_PersistentInstance))
+	{
+		return true;
+	}
+	else if (CastField<FSoftObjectProperty>(Property) || CastField<FWeakObjectProperty>(Property))
+	{
+		return true;
+	}
+	else if (const FStructProperty* StructProperty = CastField<FStructProperty>(Property))
+	{
+		const UScriptStruct* ScriptStruct = StructProperty->Struct;
+		if (ScriptStruct == FSoftObjectPath::StaticStruct() ||
+			ScriptStruct == FSoftClassPath::StaticStruct())
+		{
+			return true;
+		}
+	}
+
+	return false;
+}
+
+const FStructProperty* AsColorProperty(TNotNull<const FProperty*> Property)
+{
+	if (const FStructProperty* StructProperty = CastField<FStructProperty>(Property))
+	{
+		if (StructProperty->Struct->GetFName() == NAME_Color)
+		{
+			return StructProperty;
+		}
+	}
+	return nullptr;
+}
+
+bool CustomJsonSchema(
+	const FProperty* Property, const FString& ParameterDefaultString, const TSharedRef<FJsonObject>& OutputSchema)
+{
+	if (IsReferenceLikeProperty(Property))
+	{
+		if (ParameterDefaultString.IsEmpty())
+		{
+			FJsonObject::Duplicate(
+				FJsonSchemaGenerator::UStructToJsonSchemaObject<FToolsetReference>(), OutputSchema);
+			UClass* ClassType = nullptr;
+			bool IsClass = false;
+			if (const FClassProperty* ClassProperty = CastField<FClassProperty>(Property))
+			{
+				ClassType = ClassProperty->MetaClass;
+				IsClass = true;
+			}
+			else if (const FObjectProperty* ObjectProperty = CastField<FObjectProperty>(Property))
+			{
+				ClassType = ObjectProperty->PropertyClass;
+			}
+			else if (const FSoftObjectProperty* SoftProperty = CastField<FSoftObjectProperty>(Property))
+			{
+				ClassType = SoftProperty->PropertyClass;
+			}
+			else if (const FWeakObjectProperty* WeakProperty = CastField<FWeakObjectProperty>(Property))
+			{
+				ClassType = WeakProperty->PropertyClass;
+			}
+			else if (const FStructProperty* StructProperty = CastField<FStructProperty>(Property))
+			{
+				const UScriptStruct* ScriptStruct = StructProperty->Struct;
+				if (ScriptStruct == FSoftObjectPath::StaticStruct())
+				{
+					ClassType = UObject::StaticClass();
+				}
+				else if (ScriptStruct == FSoftClassPath::StaticStruct())
+				{
+					ClassType = UObject::StaticClass();
+					IsClass = true;
+				}
+			}
+			check(ClassType);
+			FString TypeInfo = IsClass ? FString(TEXT("Class@")) : FString();
+			TypeInfo += ClassType->GetPathName();
+			OutputSchema->SetStringField(TEXT("title"), TypeInfo);
+		}
+		else
+		{
+			if (ParameterDefaultString == FString(TEXT("None")))
+			{
+				OutputSchema->SetField(TEXT("default"), MakeShared<FJsonValueNull>());
+			}
+			else
+			{
+				FToolsetReference Reference;
+				Reference.RefPath = ParameterDefaultString;
+				OutputSchema->SetObjectField(
+					TEXT("default"),
+					FJsonObjectConverter::UStructToJsonObject(Reference));
+			}
+		}
+		return true;
+	}
+	return false;
+}
+
+bool CustomImport(const TSharedPtr<FJsonValue>& JsonValue, FProperty* Property, void* Value)
+{	
+	if (IsReferenceLikeProperty(Property))
+	{
+		FString RefPath;
+		if (JsonValue->Type == EJson::Object)
+		{						
+			JsonValue->AsObject()->TryGetStringField(TEXT("refPath"), RefPath);						
+		}
+		else if (JsonValue->Type == EJson::String)
+		{
+			RefPath = JsonValue->AsString();
+		}
+		if (!RefPath.IsEmpty() && Property->ImportText_Direct(*RefPath, Value, nullptr, 0))
+		{
+			return true;
+		}
+	}
+	return false;
+}
+
+TSharedPtr<FJsonValue> CustomExport(FProperty* Property, const void* Value)
+{
+	if (IsReferenceLikeProperty(Property))
+	{
+		FToolsetReference Reference;
+		if (FObjectProperty* ObjectProperty = CastField<FObjectProperty>(Property))
+		{
+			if (const UObject* Object = ObjectProperty->GetObjectPropertyValue(Value))
+			{
+				Reference.RefPath = Object->GetPathName();
+			}
+			else
+			{
+				return nullptr;
+			}
+		}
+		else if (FSoftObjectProperty* SoftObjectProperty = CastField<FSoftObjectProperty>(Property))
+		{
+			if (const UObject* Object = SoftObjectProperty->GetObjectPropertyValue(Value))
+			{
+				Reference.RefPath = Object->GetPathName();
+			}
+			else
+			{
+				return nullptr;
+			}
+		}
+		else if (FWeakObjectProperty* WeakObjectProperty = CastField<FWeakObjectProperty>(Property))
+		{
+			if (const UObject* Object = WeakObjectProperty->GetObjectPropertyValue(Value))
+			{
+				Reference.RefPath = Object->GetPathName();
+			}
+			else
+			{
+				return nullptr;
+			}
+		}
+		else if (const FStructProperty* StructProperty = CastField<FStructProperty>(Property))
+		{
+			const UScriptStruct* ScriptStruct = StructProperty->Struct;
+			if (ScriptStruct == FSoftObjectPath::StaticStruct())
+			{
+				const FSoftObjectPath* SoftObjectPath = reinterpret_cast<const FSoftObjectPath*>(Value);
+				Reference.RefPath = SoftObjectPath->ToString();
+			}
+			else if (ScriptStruct == FSoftClassPath::StaticStruct())
+			{
+				const FSoftClassPath* SoftClassPath = reinterpret_cast<const FSoftClassPath*>(Value);
+				Reference.RefPath = SoftClassPath->ToString();
+			}
+			else
+			{
+				check(false);
+			}
+		}
+		else
+		{
+			check(false);
+		}
+		return MakeShared<FJsonValueObject>(FJsonObjectConverter::UStructToJsonObject(Reference));
+	}
+	
+	return nullptr;
+}
+
+TSharedPtr<FJsonObject> FToolsetJson::PropertyToJsonSchema(TNotNull<const FProperty*> Property)
+{
+	FJsonSchemaPropertyFilter::CustomCallback CustomCallback;
+	CustomCallback.BindStatic(CustomJsonSchema);
+	FJsonSchemaPropertyFilter PropertyFilter;
+	PropertyFilter.CustomCb = &CustomCallback;
+	return FJsonSchemaGenerator::FPropertyToJsonSchemaObject(Property, PropertyFilter);
+}
+
+TSharedPtr<FJsonValue> FToolsetJson::PropertyToJsonData(
+	TNotNull<FProperty*> Property, TNotNull<const void*> Value)
+{
+	FJsonObjectConverter::CustomExportCallback CustomExportCallback;
+	CustomExportCallback.BindStatic(CustomExport);
+	return FJsonObjectConverter::UPropertyToJsonValue(Property, Value, 0, 0, &CustomExportCallback);
+}
+
+bool FToolsetJson::JsonDataToProperty(
+	const TSharedPtr<FJsonValue>& JsonValue, TNotNull<FProperty*> Property, TNotNull<void*> OutValue)
+{
+	FJsonObjectConverter::CustomImportCallback CustomImportCallback;
+	CustomImportCallback.BindStatic(CustomImport);
+	return FJsonObjectConverter::JsonValueToUProperty(
+		JsonValue, Property, OutValue, 0, 0, false, nullptr, &CustomImportCallback);
+}
+
+TSharedPtr<FJsonObject> FToolsetJson::StructToJsonSchema(
+	TNotNull<const UStruct*> StructDefinition, bool UserVisiblePropertiesOnly)
+{
+	FJsonSchemaPropertyFilter::CustomCallback CustomCallback;
+	CustomCallback.BindStatic(CustomJsonSchema);
+	FJsonSchemaPropertyFilter PropertyFilter;
+	PropertyFilter.CustomCb = &CustomCallback;
+	PropertyFilter.CheckFlags = GetCheckFlags(UserVisiblePropertiesOnly);
+	PropertyFilter.SkipFlags = CPF_Deprecated;
+	return FJsonSchemaGenerator::UStructToJsonSchemaObject(StructDefinition, PropertyFilter);
+}
+
+TSharedPtr<FJsonObject> FToolsetJson::StructToJsonData(
+	TNotNull<const UStruct*> StructDefinition, TNotNull<const void*> Struct, bool UserVisiblePropertiesOnly)
+{
+	FJsonObjectConverter::CustomExportCallback CustomExportCallback;
+	CustomExportCallback.BindStatic(CustomExport);
+	TSharedPtr<FJsonObject> JsonObject = MakeShared<FJsonObject>();
+	EPropertyFlags CheckFlags = GetCheckFlags(UserVisiblePropertiesOnly);
+	EPropertyFlags SkipFlags = CPF_Deprecated;
+	if (FJsonObjectConverter::UStructToJsonObject(
+		StructDefinition, Struct, JsonObject.ToSharedRef(), CheckFlags, SkipFlags, &CustomExportCallback))
+	{
+		return JsonObject;
+	}
+	return nullptr;
+}
+
+bool FToolsetJson::JsonDataToStruct(
+	const TSharedRef<FJsonObject>& JsonObject, TNotNull<const UStruct*> StructDefinition, TNotNull<void*> OutStruct)
+{
+	FJsonObjectConverter::CustomImportCallback CustomImportCallback;
+	CustomImportCallback.BindStatic(CustomImport);
+	return FJsonObjectConverter::JsonObjectToUStruct(
+		JsonObject, StructDefinition, OutStruct, 0, 0, false, nullptr, &CustomImportCallback);
+}
diff --git a/Engine/Plugins/Experimental/ToolsetRegistry/Source/ToolsetRegistry/Private/ToolsetJson.h b/Engine/Plugins/Experimental/ToolsetRegistry/Source/ToolsetRegistry/Private/ToolsetJson.h
new file mode 100644
--- /dev/null
+++ b/Engine/Plugins/Experimental/ToolsetRegistry/Source/ToolsetRegistry/Private/ToolsetJson.h
@@ -0,0 +1,83 @@
+// Copyright Epic Games, Inc. All Rights Reserved. 
+
+#pragma once
+
+#include "CoreMinimal.h"
+#include "Templates/SharedPointer.h"
+
+#include "ToolsetJson.generated.h"
+
+class FJsonObject;
+class FJsonValue;
+
+/// Represents a reference to a UObject or UClass.
+USTRUCT(BlueprintType)
+struct FToolsetReference
+{
+	GENERATED_BODY()
+public:
+	/// The reference stored as a soft path string.
+	UPROPERTY(BlueprintReadWrite, Category="Toolset")
+	FString RefPath;
+};
+
+/// Contains ToolsetRegistry specific wrappers for JSON schema and data generation.
+/// These wrappers perform custom transformations on specific property types, like
+/// converting UObject* and UClass* to FToolsetReference (vs the default in
+/// JsonUtilities which is a string).
+class FToolsetJson
+{
+public:
+	/*
+	 * Converts an FProperty to JSON schema.
+	 * @param Property The FProperty to convert.
+	 * @return JSON schema that represents the FProperty.
+	 */
+	static TSharedPtr<FJsonObject> PropertyToJsonSchema(TNotNull<const FProperty*> Property);
+	/*
+	 * Gets an FProperty on an instance to JSON data.
+	 * @param Property The FProperty to convert.
+	 * @param Value The object or struct to get the value from. 
+	 * @return JSON data that represents the value of the FProperty.
+	 */
+	static TSharedPtr<FJsonValue> PropertyToJsonData(
+		TNotNull<FProperty*> Property, TNotNull<const void*> Value);
+	/*
+	 * Set an FProperty on an instance from JSON data.
+	 * @param JsonValue The JSON data to convert.
+	 * @param Property The FProperty to apply the data to.
+	 * @param OutValue The object or struct to modify.
+	 * @return True if the data was converted successfully.
+	 */
+	static bool JsonDataToProperty(
+		const TSharedPtr<FJsonValue>& JsonValue, TNotNull<FProperty*> Property, TNotNull<void*> OutValue);
+	
+	/*
+	 * Converts a UStruct to JSON schema.
+	 * @param StructDefinition The UStruct type to convert.
+	 * @param UserVisiblePropertiesOnly Whether to export all properties or only the ones that appear in editor.
+	 * @return JSON schema that represents the UStruct.
+	 */
+	static TSharedPtr<FJsonObject> StructToJsonSchema(
+		TNotNull<const UStruct*> StructDefinition, bool UserVisiblePropertiesOnly = false);
+	/*
+	 * Converts a UStruct instance to JSON data.
+	 * @param StructDefinition The UStruct type to convert.
+	 * @param Struct The UStruct instance to convert.
+	 * @param UserVisiblePropertiesOnly Whether to export all properties or only the ones that appear in editor.
+	 * @return JSON data that represents the UStruct instance.
+	 */
+	static TSharedPtr<FJsonObject> StructToJsonData(
+		TNotNull<const UStruct*> StructDefinition, TNotNull<const void*> Struct,
+		bool UserVisiblePropertiesOnly = false);
+	/*
+	 * Updates a UStruct instance from JSON data.
+	 * @param JsonObject The JSON data to convert.
+	 * @param StructDefinition The UStruct type.
+	 * @param OutStruct The UStruct instance to modify.
+	 * @return True if the data was applied successfully.
+	 */
+	static bool JsonDataToStruct(
+		const TSharedRef<FJsonObject>& JsonObject, TNotNull<const UStruct*> StructDefinition,
+		TNotNull<void*> OutStruct);
+};
diff --git a/Engine/Plugins/Experimental/ToolsetRegistry/Source/ToolsetRegistry/Private/ToolsetLibrary.cpp b/Engine/Plugins/Experimental/ToolsetRegistry/Source/ToolsetRegistry/Private/ToolsetLibrary.cpp
--- a/Engine/Plugins/Experimental/ToolsetRegistry/Source/ToolsetRegistry/Private/ToolsetLibrary.cpp
+++ b/Engine/Plugins/Experimental/ToolsetRegistry/Source/ToolsetRegistry/Private/ToolsetLibrary.cpp
@@ -2,20 +2,25 @@
 
 #include "ToolsetLibrary.h"
 
-#include "JsonObjectConverter.h"
-#include "JsonSchema/JsonSchemaGenerator.h"
+#include "Dom/JsonObject.h"
 #include "LevelEditorViewport.h"
 #include "UObject/PropertyAccessUtil.h"
+#include "Serialization/JsonSerializer.h"
+#include "Serialization/JsonReader.h"
+#include "Serialization/JsonWriter.h"
+
+#include "ToolsetJson.h"
 
 #include UE_INLINE_GENERATED_CPP_BY_NAME(ToolsetLibrary)
 
 FString UToolsetLibraryTemp::ListClassProperties(UClass* Class)
 {
+	if (!Class)
+	{
+		return FString();
+	}
 	FString JsonProperties;
-	FJsonSchemaPropertyFilter Filter;
-	Filter.CheckFlags |= CPF_BlueprintVisible;
-	Filter.SkipFlags |= CPF_Deprecated;
-	TSharedPtr<FJsonObject> ClassSchema = FJsonSchemaGenerator::UStructToJsonSchemaObject(Class, Filter);
+	TSharedPtr<FJsonObject> ClassSchema = FToolsetJson::StructToJsonSchema(Class, true);
 	const TSharedPtr<FJsonObject>* ClassProperties = nullptr;
 	if (ClassSchema->TryGetObjectField(TEXT("properties"), ClassProperties))
 	{
@@ -33,8 +38,9 @@ FString UToolsetLibraryTemp::GetObjectProperties(const UObject* Object, const TA
 	{
 		FProperty* Property = Object->GetClass()->FindPropertyByName(PropertyName);
 		if(Property)
-		{			
-			PropertyData->SetField(PropertyName.ToString(), FJsonObjectConverter::UPropertyToJsonValue(Property, Property->ContainerPtrToValuePtr<void>(Object)));
+		{
+			PropertyData->SetField(PropertyName.ToString(), FToolsetJson::PropertyToJsonData(
+				Property, Property->ContainerPtrToValuePtr<void>(Object)));
 		}
 	}
 	FString JsonProperties;
@@ -59,10 +65,11 @@ bool UToolsetLibraryTemp::SetObjectProperties(UObject* Object, const FString& Pr
 			FProperty* Property = ObjectClass->FindPropertyByName(FName(Key));
 			if (Property)
 			{
-				TUniquePtr<FPropertyAccessChangeNotify> ChangeNotify = PropertyAccessUtil::BuildBasicChangeNotify(Property, Object, EPropertyAccessChangeNotifyMode::Always, EPropertyChangeType::ValueSet);
+				TUniquePtr<FPropertyAccessChangeNotify> ChangeNotify = PropertyAccessUtil::BuildBasicChangeNotify(
+					Property, Object, EPropertyAccessChangeNotifyMode::Always, EPropertyChangeType::ValueSet);
 				PropertyAccessUtil::EmitPreChangeNotify(ChangeNotify.Get(), false);
 				void* SettingsPropPtr = Property->ContainerPtrToValuePtr<void>(Object);
-				FJsonObjectConverter::JsonValueToUProperty(PropertyData->TryGetField(Key), Property, SettingsPropPtr);
+				FToolsetJson::JsonDataToProperty(PropertyData->TryGetField(Key), Property, SettingsPropPtr);
 				PropertyAccessUtil::EmitPostChangeNotify(ChangeNotify.Get(), false);
 			}
 			else
diff --git a/Engine/Plugins/Experimental/ToolsetRegistry/Source/ToolsetRegistry/Private/ToolsetRegistryToolCall.cpp b/Engine/Plugins/Experimental/ToolsetRegistry/Source/ToolsetRegistry/Private/ToolsetRegistryToolCall.cpp
--- a/Engine/Plugins/Experimental/ToolsetRegistry/Source/ToolsetRegistry/Private/ToolsetRegistryToolCall.cpp
+++ b/Engine/Plugins/Experimental/ToolsetRegistry/Source/ToolsetRegistry/Private/ToolsetRegistryToolCall.cpp
@@ -2,16 +2,15 @@
 
 #include "ToolsetRegistryToolCall.h"
 
+#include "Async/Async.h"
 #include "Engine/Engine.h"
 #include "Engine/World.h"
 #include "Engine/EngineTypes.h"
 #include "Modules/ModuleManager.h"
-#include "JsonDomBuilder.h"
-#include "Async/Async.h"
 #include "Serialization/JsonReader.h"
 #include "Serialization/JsonSerializer.h"
+#include "ToolsetJson.h"
 #include "UObject/FieldIterator.h"
-#include "JsonSchema/JsonSchemaGenerator.h"
 
 
 namespace UE::ToolsetRegistry
@@ -285,7 +284,7 @@ namespace UE::ToolsetRegistry
 		InstanceObject(InInstanceObject),
 		Function(InFunction)
 	{
-		FunctionSchemaJsonObject = FJsonSchemaGenerator::UStructToJsonSchemaObject(Function);
+		FunctionSchemaJsonObject = FToolsetJson::StructToJsonSchema(Function);
 		check(FunctionSchemaJsonObject);
 	}
 
@@ -350,7 +349,7 @@ namespace UE::ToolsetRegistry
 		};
 
 		// Copy the function input params memory into the function's frame memory.
-		if (!FJsonObjectConverter::JsonObjectToUStruct(
+		if (FrameMemory && !FToolsetJson::JsonDataToStruct(
 			FunctionInputParamsJsonObject.ToSharedRef(), Function, FrameMemory))
 		{
 			// Caller error. Ill-formed incoming function input params JSON. 
@@ -419,7 +418,7 @@ namespace UE::ToolsetRegistry
 
 		// Get the return value, if we have one.
 		TSharedPtr<FJsonValue> ReturnJsonValue = Function->GetReturnProperty() ?
-			FJsonObjectConverter::UPropertyToJsonValue(Function->GetReturnProperty(), ReturnValueAddress) :
+			FToolsetJson::PropertyToJsonData(Function->GetReturnProperty(), ReturnValueAddress) :
 			MakeShared<FJsonValueNull>();
 
 		FString ErrorMessage;
diff --git a/Engine/Plugins/Experimental/ToolsetRegistry/Source/ToolsetRegistry/Private/ToolsetRegistryToolset.cpp b/Engine/Plugins/Experimental/ToolsetRegistry/Source/ToolsetRegistry/Private/ToolsetRegistryToolset.cpp
--- a/Engine/Plugins/Experimental/ToolsetRegistry/Source/ToolsetRegistry/Private/ToolsetRegistryToolset.cpp
+++ b/Engine/Plugins/Experimental/ToolsetRegistry/Source/ToolsetRegistry/Private/ToolsetRegistryToolset.cpp
@@ -5,7 +5,6 @@
 #include "Async/Future.h"
 #include "Containers/UnrealString.h"
 #include "Dom/JsonObject.h"
-#include "JsonSchema/JsonSchemaGenerator.h"
 #include "Kismet/BlueprintFunctionLibrary.h"
 #include "Misc/CoreMisc.h"
 #include "Misc/Optional.h"
@@ -15,6 +14,7 @@
 #include "Blueprint/BlueprintExceptionInfo.h"
 #include "ToolsetRegistryDelegateHandle.h"
 
+#include "ToolsetJson.h"
 #include "ToolsetRegistryToolCall.h"
 #include "ToolsetRegistryModule.h"
 
@@ -193,7 +193,7 @@ namespace UE::ToolsetRegistry
 		{
 			check(Function);
 			
-			TSharedPtr<FJsonObject> ToolEntry = FJsonSchemaGenerator::UStructToJsonSchemaObject(Function);			
+			TSharedPtr<FJsonObject> ToolEntry = FToolsetJson::StructToJsonSchema(Function);			
 
 			// Modify tool name from "ClassName.FunctionName" to "ToolsetName.FunctionName".
 			FString ToolName = ToolEntry->GetStringField(TEXT("name"));
