[AI Assistant] Add JSON conversion library to toolset registry.

Provides a consistent way to convert from JSON strings to JSON objects / values and vice
versa.

#Synced-CL 50892692

#Synced-CL 50892692

#Synced-CL 50892692

diff --git a/Engine/Plugins/Experimental/ToolsetRegistry/Source/ToolsetRegistry/Private/Tests/JsonConversionTest.cpp b/Engine/Plugins/Experimental/ToolsetRegistry/Source/ToolsetRegistry/Private/Tests/JsonConversionTest.cpp
new file mode 100644
--- /dev/null
+++ b/Engine/Plugins/Experimental/ToolsetRegistry/Source/ToolsetRegistry/Private/Tests/JsonConversionTest.cpp
@@ -0,0 +1,129 @@
+// Copyright Epic Games, Inc. All Rights Reserved.
+
+#include "Containers/UnrealString.h"
+#include "Dom/JsonObject.h"
+#include "Misc/AutomationTest.h"
+#include "Templates/SharedPointer.h"
+
+#include "ToolsetRegistry/JsonConversion.h"
+
+#if WITH_DEV_AUTOMATION_TESTS
+
+BEGIN_DEFINE_SPEC(
+	FToolsetRegistryJsonConversionSpec,
+	"AI.ToolsetRegistry.JsonConversionTest",
+	EAutomationTestFlags::EditorContext | EAutomationTestFlags::EngineFilter)
+END_DEFINE_SPEC(FToolsetRegistryJsonConversionSpec)
+
+void FToolsetRegistryJsonConversionSpec::Define()
+{
+	using namespace UE::ToolsetRegistry::Internal;
+
+	Describe(TEXT("JsonStringToJsonObject"), [this]()
+	{
+		It(TEXT("should return no object from an empty string"), [this]()
+		{
+			TSharedPtr<FJsonObject> JsonObject = JsonStringToJsonObject(TEXT(""));
+			TestFalse(TEXT("IsValid"), JsonObject.IsValid());
+		});
+
+		It(TEXT("should return no object given an invalid string"), [this]()
+		{
+			TSharedPtr<FJsonObject> JsonObject =
+				JsonStringToJsonObject(TEXT(R"json("a_value")json"));
+			TestFalse(TEXT("IsValid"), JsonObject.IsValid());
+		});
+
+		It(TEXT("should return an object from a valid object string"), [this]()
+		{
+			TSharedPtr<FJsonObject> JsonObject =
+				JsonStringToJsonObject(TEXT(R"json({"hello": "world"})json"));
+			if (!TestTrue(TEXT("IsValid"), JsonObject.IsValid())) return;
+			
+			TSharedPtr<FJsonValue> JsonValue = JsonObject->TryGetField(TEXT("hello"));
+			if (!TestTrue(TEXT("Has 'hello' field"), JsonValue.IsValid())) return;
+
+			TestEqual(TEXT("Field value"), JsonValue->AsString(), TEXT("world"));
+		});
+	});
+
+	Describe(TEXT("JsonObjectOrEmpty"), [this]()
+	{
+		It(TEXT("Should return an existing object when provided an object"), [this]()
+		{
+			TSharedPtr<FJsonObject> JsonObject =
+				JsonStringToJsonObject(TEXT(R"json({"hello": "world"})json"));
+			TestEqual(TEXT("Json object"), &JsonObjectOrEmpty(JsonObject).Get(), JsonObject.Get());
+		});
+
+		It(TEXT("Should return an empty when provided no object"), [this]()
+		{
+			TSharedRef<FJsonObject> JsonObject = JsonObjectOrEmpty(TSharedPtr<FJsonObject>());
+			TestTrue(TEXT("Json object with no values"), JsonObject->Values.IsEmpty());
+		});
+	});
+
+	Describe(TEXT("JsonStringToJsonValue"), [this]()
+	{
+		It(TEXT("should return no value from an empty string"), [this]()
+		{
+			TSharedPtr<FJsonValue> JsonValue = JsonStringToJsonValue(TEXT(""));
+			TestFalse(TEXT("IsValid"), JsonValue.IsValid());
+		});
+
+		It(TEXT("should return no value given an invalid string"), [this]()
+		{
+			TSharedPtr<FJsonValue> JsonValue =
+				JsonStringToJsonValue(TEXT(R"json(a_value)json"));
+			TestFalse(TEXT("IsValid"), JsonValue.IsValid());
+		});
+
+		It(TEXT("should return a value from a valid value string"), [this]()
+		{
+			TSharedPtr<FJsonValue> JsonValue =
+				JsonStringToJsonValue(TEXT(R"json("hello")json"));
+			if (!TestTrue(TEXT("IsValid"), JsonValue.IsValid())) return;
+			TestEqual(TEXT("Value"), JsonValue->AsString(), TEXT("hello"));
+		});
+	});
+
+	Describe(TEXT("JsonToString"), [this]()
+	{
+		It(TEXT("should return an empty object string from an empty object"), [this]()
+		{
+			TSharedRef<FJsonObject> JsonObject = MakeShared<FJsonObject>();
+			TestEqual(
+				TEXT("No JSON object"), JsonToString(JsonObject), TEXT("{}"));
+		});
+
+		It(TEXT("should return a JSON object string from a JSON object"), [this]()
+		{
+			TSharedRef<FJsonObject> JsonObject = MakeShared<FJsonObject>();
+			JsonObject->SetStringField(TEXT("hello"), TEXT("world"));
+			TestEqual(
+				TEXT("JSON object"),
+				JsonToString(JsonObject),
+				TEXT(R"json({"hello":"world"})json"));
+		});
+
+		It(TEXT("should return a string value from a JSON value string"), [this]()
+		{
+			TSharedRef<FJsonValue> JsonValue = MakeShared<FJsonValueString>(TEXT("yo"));
+			TestEqual(
+				TEXT("JSON object"),
+				JsonToString(JsonValue),
+				FString::Printf(TEXT(R"json("%s")json"), *JsonValue->AsString()));
+		});
+
+		It(TEXT("should return an null string from a null JSON value"), [this]()
+		{
+			TSharedRef<FJsonValue> JsonValue = MakeShared<FJsonValueNull>();
+			TestEqual(
+				TEXT("JSON object"),
+				JsonToString(JsonValue),
+				TEXT("null"));
+		});
+	});
+}
+
+#endif  // WITH_DEV_AUTOMATION_TESTS
\ No newline at end of file
diff --git a/Engine/Plugins/Experimental/ToolsetRegistry/Source/ToolsetRegistry/Private/ToolsetRegistry/JsonConversion.cpp b/Engine/Plugins/Experimental/ToolsetRegistry/Source/ToolsetRegistry/Private/ToolsetRegistry/JsonConversion.cpp
new file mode 100644
--- /dev/null
+++ b/Engine/Plugins/Experimental/ToolsetRegistry/Source/ToolsetRegistry/Private/ToolsetRegistry/JsonConversion.cpp
@@ -0,0 +1,87 @@
+// Copyright Epic Games, Inc. All Rights Reserved.
+
+#include "ToolsetRegistry/JsonConversion.h"
+
+#include "Misc/AssertionMacros.h"
+#include "Policies/CondensedJsonPrintPolicy.h"
+#include "Templates/Function.h"
+#include "Serialization/JsonReader.h"
+#include "Serialization/JsonWriter.h"
+#include "Serialization/JsonSerializer.h"
+
+namespace UE::ToolsetRegistry::Internal
+{
+	TSharedPtr<FJsonObject> JsonStringToJsonObject(const FString& JsonString)
+	{
+		TSharedPtr<FJsonObject> ReturnJsonObject;
+		if (!JsonString.IsEmpty())
+		{
+			const TSharedRef<TJsonReader<>> JsonReader =
+				TJsonReaderFactory<>::Create(JsonString);
+			TSharedPtr<FJsonObject> JsonObject;
+			if (FJsonSerializer::Deserialize(JsonReader, JsonObject))
+			{
+				ReturnJsonObject = JsonObject;
+			}
+		}
+		return ReturnJsonObject;
+	}
+
+	TSharedRef<FJsonObject> JsonObjectOrEmpty(TSharedPtr<FJsonObject> JsonObject)
+	{
+		return JsonObject ? JsonObject.ToSharedRef() : MakeShared<FJsonObject>();
+	}
+
+	TSharedPtr<FJsonValue> JsonStringToJsonValue(const FString& JsonString)
+	{
+		// Work around JSON deserialization's lack of support for loading values from strings
+		// by wrapping in a temporary object.
+		TSharedPtr<FJsonValue> ReturnJsonValue;
+		if (!JsonString.IsEmpty())
+		{
+			static const FString Key = TEXT("__key");
+			TSharedPtr<FJsonObject> JsonObject =
+				JsonStringToJsonObject(
+					FString::Printf(TEXT(R"json({"%s": %s})json"), *Key, *JsonString));
+			if (JsonObject)
+			{
+				ReturnJsonValue = JsonObject->Values.FindAndRemoveChecked(Key);
+			}
+		}
+		return ReturnJsonValue;
+	}
+
+	namespace
+	{
+		using FJsonPrintPolicy = TCondensedJsonPrintPolicy<TCHAR>;
+		using FJsonWriterFactory = TJsonWriterFactory<FJsonPrintPolicy::CharType, FJsonPrintPolicy>;
+		using FJsonWriter = TJsonWriter<FJsonPrintPolicy::CharType, FJsonPrintPolicy>;
+
+		FString WithJsonWriter(TFunction<void(TSharedRef<FJsonWriter>)> Apply)
+		{
+			FString JsonString;
+			TSharedRef<FJsonWriter> JsonWriter = FJsonWriterFactory::Create(&JsonString);
+			Apply(JsonWriter);
+			return JsonString;
+		}
+	}
+
+	FString JsonToString(const TSharedRef<FJsonObject> JsonObject)
+	{
+		return WithJsonWriter(
+			[JsonObject](auto JsonWriter) -> void
+			{
+				FJsonSerializer::Serialize(JsonObject, JsonWriter);
+			});
+	}
+
+	FString JsonToString(const TSharedRef<FJsonValue> JsonValue)
+	{
+		FString JsonString = WithJsonWriter(
+			[JsonValue](auto JsonWriter) -> void
+			{
+				FJsonSerializer::Serialize(JsonValue, TEXT(""), JsonWriter);
+			});
+		return JsonString;
+	}
+}
\ No newline at end of file
diff --git a/Engine/Plugins/Experimental/ToolsetRegistry/Source/ToolsetRegistry/Private/ToolsetRegistry/JsonConversion.h b/Engine/Plugins/Experimental/ToolsetRegistry/Source/ToolsetRegistry/Private/ToolsetRegistry/JsonConversion.h
new file mode 100644
--- /dev/null
+++ b/Engine/Plugins/Experimental/ToolsetRegistry/Source/ToolsetRegistry/Private/ToolsetRegistry/JsonConversion.h
@@ -0,0 +1,27 @@
+// Copyright Epic Games, Inc. All Rights Reserved.
+
+#pragma once
+
+#include "Containers/UnrealString.h"
+#include "Dom/JsonObject.h"
+#include "Templates/SharedPointer.h"
+
+namespace UE::ToolsetRegistry::Internal
+{
+	// Convert a JSON string to a JSON object or return a null pointer if the string cannot be
+	// parsed or isn't an object.
+	TSharedPtr<FJsonObject> JsonStringToJsonObject(const FString& JsonString);
+
+	// Return either the provided JSON object or return an empty object.
+	TSharedRef<FJsonObject> JsonObjectOrEmpty(TSharedPtr<FJsonObject> JsonObject);
+
+	// Convert a JSON string to a JSON value or return a null pointer if the string cannot be
+	// parsed.
+	TSharedPtr<FJsonValue> JsonStringToJsonValue(const FString& JsonString);
+
+	// Convert a JSON object to a string.
+	FString JsonToString(const TSharedRef<FJsonObject> JsonObject);
+
+	// Convert a JSON value to a string.
+	FString JsonToString(const TSharedRef<FJsonValue> JsonValue);
+}
\ No newline at end of file
