[AI Assistant] Use JSON conversion methods in FObjectFunctionToolCall

#tests

#Synced-CL 50949381

diff --git a/Engine/Plugins/Experimental/ToolsetRegistry/Source/ToolsetRegistry/Private/ToolsetRegistryToolCall.cpp b/Engine/Plugins/Experimental/ToolsetRegistry/Source/ToolsetRegistry/Private/ToolsetRegistryToolCall.cpp
--- a/Engine/Plugins/Experimental/ToolsetRegistry/Source/ToolsetRegistry/Private/ToolsetRegistryToolCall.cpp
+++ b/Engine/Plugins/Experimental/ToolsetRegistry/Source/ToolsetRegistry/Private/ToolsetRegistryToolCall.cpp
@@ -12,36 +12,14 @@
 #include "ToolsetJson.h"
 #include "UObject/FieldIterator.h"
 
+#include "ToolsetRegistry/JsonConversion.h"
 
 namespace UE::ToolsetRegistry
 {
+	using namespace UE::ToolsetRegistry::Internal;
+
 	namespace
 	{
-		TSharedRef<FJsonObject> JsonStringToValidJsonObject(const FString& JsonString)
-		{
-			if (!JsonString.IsEmpty())
-			{
-				const TSharedRef<TJsonReader<>> JsonReader =
-					TJsonReaderFactory<>::Create(JsonString);
-				TSharedPtr<FJsonObject> JsonObject;
-				if (FJsonSerializer::Deserialize(JsonReader, JsonObject))
-				{
-					return JsonObject.ToSharedRef();
-				}
-			}
-
-			return MakeShared<FJsonObject>();
-		}
-
-		FString JsonObjectToString(const TSharedRef<FJsonObject>& JsonObject)
-		{
-			FString JsonString;
-			TSharedRef<TJsonWriter<>> JsonWriter = TJsonWriterFactory<>::Create(&JsonString);
-			FJsonSerializer::Serialize(JsonObject, JsonWriter);
-
-			return JsonString;
-		}
-
 		TSharedRef<FJsonObject> FunctionInputParamsJsonToValidJsonObject(
 			const TOptional<FObjectFunctionToolCall::FFunctionInputParamsJson>&
 				FunctionInputParamsJson)
@@ -60,7 +38,8 @@ namespace UE::ToolsetRegistry
 				}
 				else
 				{
-					return JsonStringToValidJsonObject(FunctionInputParamsJson->Get<FString>());
+					return JsonObjectOrEmpty(
+						JsonStringToJsonObject(FunctionInputParamsJson->Get<FString>()));
 				}
 			}
 
@@ -145,7 +124,7 @@ namespace UE::ToolsetRegistry
 							"Function schema Json -\n"
 							"%s"),
 							*FunctionDebugName,
-							*JsonObjectToString(FunctionSchemaInputSchemaJsonObject));
+							*JsonToString(FunctionSchemaInputSchemaJsonObject));
 					}
 
 					return false;
@@ -170,8 +149,8 @@ namespace UE::ToolsetRegistry
 								"%s"),
 								*FunctionDebugName,
 								*RequiredParamName,
-								*JsonObjectToString(FunctionSchemaInputSchemaJsonObject),
-								*JsonObjectToString(FunctionInputParamsJsonObject));
+								*JsonToString(FunctionSchemaInputSchemaJsonObject),
+								*JsonToString(FunctionInputParamsJsonObject));
 						}
 
 						return false;
@@ -230,7 +209,7 @@ namespace UE::ToolsetRegistry
 								"%s"),
 								*FunctionDebugName,
 								*InputParamPropertyName,
-								*JsonObjectToString(FunctionInputParamsJsonObject));
+								*JsonToString(FunctionInputParamsJsonObject));
 						}
 
 						return false;
@@ -254,7 +233,7 @@ namespace UE::ToolsetRegistry
 									"%s"),
 									*FunctionDebugName,
 									*InputParamPropertyName,
-									*JsonObjectToString(FunctionInputParamsJsonObject));
+									*JsonToString(FunctionInputParamsJsonObject));
 							}
 
 							return false;
@@ -361,7 +340,7 @@ namespace UE::ToolsetRegistry
 				"Function input params Json -\n"
 				"%s"),
 				*Function->GetName(),
-				*JsonObjectToString(FunctionInputParamsJsonObject.ToSharedRef())));
+				*JsonToString(FunctionInputParamsJsonObject.ToSharedRef())));
 		}
 
 		// In case this is a Blueprint function that has a world context object parameter, 
