[AI Assistant] Use JSON conversion in all tool registry modules.

#Synced-CL 50952660

diff --git a/Engine/Plugins/Experimental/ToolsetRegistry/Source/ToolsetRegistry/Private/Tests/ToolsetRegistryToolCallTest.cpp b/Engine/Plugins/Experimental/ToolsetRegistry/Source/ToolsetRegistry/Private/Tests/ToolsetRegistryToolCallTest.cpp
--- a/Engine/Plugins/Experimental/ToolsetRegistry/Source/ToolsetRegistry/Private/Tests/ToolsetRegistryToolCallTest.cpp
+++ b/Engine/Plugins/Experimental/ToolsetRegistry/Source/ToolsetRegistry/Private/Tests/ToolsetRegistryToolCallTest.cpp
@@ -11,6 +11,7 @@
 
 #include "ToolsetRegistryToolCall.h"
 #include "ToolsetRegistryModule.h"
+#include "ToolsetRegistry/JsonConversion.h"
 
 PRAGMA_DISABLE_DEPRECATION_WARNINGS
 
@@ -27,9 +28,8 @@ namespace UE::ToolsetRegistry
 		{
 			if (JsonObject)
 			{
-				FString JsonString;
-				const TSharedRef<TJsonWriter<>> JsonWriter = TJsonWriterFactory<>::Create(&JsonString);
-				FJsonSerializer::Serialize(JsonObject, JsonWriter);
+				FString JsonString =
+					UE::ToolsetRegistry::Internal::JsonToString(JsonObject.ToSharedRef());
 				UE_LOG(LogToolsetRegistry, Display, TEXT("Json \"%s\" is valid -\n%s"), Name, *JsonString);
 			}
 			else
diff --git a/Engine/Plugins/Experimental/ToolsetRegistry/Source/ToolsetRegistry/Private/Tests/ToolsetRegistryToolsetTest.cpp b/Engine/Plugins/Experimental/ToolsetRegistry/Source/ToolsetRegistry/Private/Tests/ToolsetRegistryToolsetTest.cpp
--- a/Engine/Plugins/Experimental/ToolsetRegistry/Source/ToolsetRegistry/Private/Tests/ToolsetRegistryToolsetTest.cpp
+++ b/Engine/Plugins/Experimental/ToolsetRegistry/Source/ToolsetRegistry/Private/Tests/ToolsetRegistryToolsetTest.cpp
@@ -20,6 +20,7 @@
 #include "BPToolsetRegistry.h"
 #include "ToolsetRegistryToolset.h"
 #include "ToolsetRegistrySubsystem.h"
+#include "ToolsetRegistry/JsonConversion.h"
 
 #if WITH_DEV_AUTOMATION_TESTS
 const FString FakeToolsetName(TEXT("ToolsetRegistry.FakeToolset"));
@@ -355,10 +356,8 @@ bool FToolsetRegistryToolsetTestGetSchema::RunTest(const FString& UnusedParamete
 
 	TSharedPtr<FJsonObject> ExpectedSchema = GetFakeToolsetExpectedSchema();
 
-	FString ExpectedSchemaJson;
-	FJsonSerializer::Serialize(
-		ExpectedSchema.ToSharedRef(),
-		TJsonWriterFactory<TCHAR, TCondensedJsonPrintPolicy<TCHAR>>::Create(&ExpectedSchemaJson));
+	FString ExpectedSchemaJson =
+		UE::ToolsetRegistry::Internal::JsonToString(ExpectedSchema.ToSharedRef());
 	FString OutputSchemaJson = Toolset->GetToolsetJsonSchema();
 	TestEqual("Toolset schema matches expected", OutputSchemaJson, ExpectedSchemaJson);
 	return true;
@@ -378,10 +377,8 @@ bool FToolsetRegistryToolsetTestGetSchemaWithDefaultParams::RunTest(const FStrin
 
 	TSharedPtr<FJsonObject> ExpectedSchema = GetFakeToolsetWithDefaultParamsExpectedSchema();
 
-	FString ExpectedSchemaJson;
-	FJsonSerializer::Serialize(
-		ExpectedSchema.ToSharedRef(),
-		TJsonWriterFactory<TCHAR, TCondensedJsonPrintPolicy<TCHAR>>::Create(&ExpectedSchemaJson));
+	FString ExpectedSchemaJson =
+		UE::ToolsetRegistry::Internal::JsonToString(ExpectedSchema.ToSharedRef());
 	FString OutputSchemaJson = Toolset->GetToolsetJsonSchema();
 	TestEqual("Toolset schema matches expected", OutputSchemaJson, ExpectedSchemaJson);
 	return true;
@@ -401,10 +398,8 @@ bool FToolsetRegistryToolsetTestGetSchemaWithArrayParams::RunTest(const FString&
 
 	TSharedPtr<FJsonObject> ExpectedSchema = GetFakeToolsetWithArrayParamsExpectedSchema();
 
-	FString ExpectedSchemaJson;
-	FJsonSerializer::Serialize(
-		ExpectedSchema.ToSharedRef(),
-		TJsonWriterFactory<TCHAR, TCondensedJsonPrintPolicy<TCHAR>>::Create(&ExpectedSchemaJson));
+	FString ExpectedSchemaJson =
+		UE::ToolsetRegistry::Internal::JsonToString(ExpectedSchema.ToSharedRef());
 	FString OutputSchemaJson = Toolset->GetToolsetJsonSchema();
 	TestEqual("Toolset schema matches expected", OutputSchemaJson, ExpectedSchemaJson);
 	return true;
@@ -424,10 +419,8 @@ bool FToolsetRegistryToolsetTestGetSchemaWithMapParams::RunTest(const FString& U
 
 	TSharedPtr<FJsonObject> ExpectedSchema = GetFakeToolsetWithMapParamsExpectedSchema();
 
-	FString ExpectedSchemaJson;
-	FJsonSerializer::Serialize(
-		ExpectedSchema.ToSharedRef(),
-		TJsonWriterFactory<TCHAR, TCondensedJsonPrintPolicy<TCHAR>>::Create(&ExpectedSchemaJson));
+	FString ExpectedSchemaJson =
+		UE::ToolsetRegistry::Internal::JsonToString(ExpectedSchema.ToSharedRef());
 	FString OutputSchemaJson = Toolset->GetToolsetJsonSchema();
 	TestEqual("Toolset schema matches expected", OutputSchemaJson, ExpectedSchemaJson);
 	return true;
@@ -447,10 +440,8 @@ bool FToolsetRegistryToolsetTestGetSchemaWithSetParams::RunTest(const FString& U
 
 	TSharedPtr<FJsonObject> ExpectedSchema = GetFakeToolsetWithSetParamsExpectedSchema();
 
-	FString ExpectedSchemaJson;
-	FJsonSerializer::Serialize(
-		ExpectedSchema.ToSharedRef(),
-		TJsonWriterFactory<TCHAR, TCondensedJsonPrintPolicy<TCHAR>>::Create(&ExpectedSchemaJson));
+	FString ExpectedSchemaJson =
+		UE::ToolsetRegistry::Internal::JsonToString(ExpectedSchema.ToSharedRef());
 	FString OutputSchemaJson = Toolset->GetToolsetJsonSchema();
 	TestEqual("Toolset schema matches expected", OutputSchemaJson, ExpectedSchemaJson);
 	return true;
diff --git a/Engine/Plugins/Experimental/ToolsetRegistry/Source/ToolsetRegistry/Private/ToolsetLibrary.cpp b/Engine/Plugins/Experimental/ToolsetRegistry/Source/ToolsetRegistry/Private/ToolsetLibrary.cpp
--- a/Engine/Plugins/Experimental/ToolsetRegistry/Source/ToolsetRegistry/Private/ToolsetLibrary.cpp
+++ b/Engine/Plugins/Experimental/ToolsetRegistry/Source/ToolsetRegistry/Private/ToolsetLibrary.cpp
@@ -4,12 +4,14 @@
 
 #include "Dom/JsonObject.h"
 #include "LevelEditorViewport.h"
+#include "Misc/AssertionMacros.h"
 #include "UObject/PropertyAccessUtil.h"
 #include "Serialization/JsonSerializer.h"
 #include "Serialization/JsonReader.h"
 #include "Serialization/JsonWriter.h"
 
 #include "ToolsetJson.h"
+#include "ToolsetRegistry/JsonConversion.h"
 
 #include UE_INLINE_GENERATED_CPP_BY_NAME(ToolsetLibrary)
 
@@ -24,8 +26,9 @@ FString UToolsetLibraryTemp::ListClassProperties(UClass* Class)
 	const TSharedPtr<FJsonObject>* ClassProperties = nullptr;
 	if (ClassSchema->TryGetObjectField(TEXT("properties"), ClassProperties))
 	{
-		TSharedRef<TJsonWriter<>> Writer = TJsonWriterFactory<>::Create(&JsonProperties);
-		FJsonSerializer::Serialize(ClassProperties->ToSharedRef(), Writer);
+		check(ClassProperties);
+		JsonProperties = UE::ToolsetRegistry::Internal::JsonToString(
+			ClassProperties->ToSharedRef());
 	}
 	return JsonProperties;
 }
@@ -43,18 +46,15 @@ FString UToolsetLibraryTemp::GetObjectProperties(const UObject* Object, const TA
 				Property, Property->ContainerPtrToValuePtr<void>(Object)));
 		}
 	}
-	FString JsonProperties;
-	TSharedRef<TJsonWriter<>> Writer = TJsonWriterFactory<>::Create(&JsonProperties);
-	FJsonSerializer::Serialize(PropertyData.ToSharedRef(), Writer);
-	return JsonProperties;
+	return UE::ToolsetRegistry::Internal::JsonToString(PropertyData.ToSharedRef());
 }
 
 bool UToolsetLibraryTemp::SetObjectProperties(UObject* Object, const FString& PropertiesJson)
 {
 	bool AllPropertiesSet = false;
-	TSharedPtr<FJsonObject> PropertyData;
-	TSharedRef<TJsonReader<>> Reader = TJsonReaderFactory<>::Create(PropertiesJson);
-	if (FJsonSerializer::Deserialize(Reader, PropertyData) && PropertyData.IsValid())
+	TSharedPtr<FJsonObject> PropertyData =
+		UE::ToolsetRegistry::Internal::JsonStringToJsonObject(PropertiesJson);
+	if (PropertyData.IsValid())
 	{
 		AllPropertiesSet = true;
 		const UClass* ObjectClass = Object->GetClass();
diff --git a/Engine/Plugins/Experimental/ToolsetRegistry/Source/ToolsetRegistry/Private/ToolsetRegistryToolset.cpp b/Engine/Plugins/Experimental/ToolsetRegistry/Source/ToolsetRegistry/Private/ToolsetRegistryToolset.cpp
--- a/Engine/Plugins/Experimental/ToolsetRegistry/Source/ToolsetRegistry/Private/ToolsetRegistryToolset.cpp
+++ b/Engine/Plugins/Experimental/ToolsetRegistry/Source/ToolsetRegistry/Private/ToolsetRegistryToolset.cpp
@@ -17,6 +17,7 @@
 #include "ToolsetJson.h"
 #include "ToolsetRegistryToolCall.h"
 #include "ToolsetRegistryModule.h"
+#include "ToolsetRegistry/JsonConversion.h"
 
 namespace UE::ToolsetRegistry
 {
@@ -164,19 +165,8 @@ namespace UE::ToolsetRegistry
 
 	FString FToolsetLibraryWrapper::GetToolsetJsonSchema() const
 	{
-		FString OutputString;
 		TSharedPtr<FJsonObject> Schema = GenerateSchema();
-
-		if (!FJsonSerializer::Serialize(
-				Schema.ToSharedRef(),
-				TJsonWriterFactory<TCHAR, TCondensedJsonPrintPolicy<TCHAR>>::Create(&OutputString)))
-		{
-			UE_LOG(
-				LogToolsetRegistry, Warning,
-				TEXT("Failed to generate valid toolset schema for toolset %s"), *GetToolsetName());
-		}
-
-		return OutputString;
+		return UE::ToolsetRegistry ::Internal::JsonToString(Schema.ToSharedRef());
 	}
 
 	TSharedPtr<FJsonObject> FToolsetLibraryWrapper::GenerateSchema() const
@@ -278,14 +268,12 @@ namespace UE::ToolsetRegistry
 		}
 		check(ToolCallResult.HasValue());
 		TSharedPtr<FJsonObject> ResultValue = ToolCallResult.StealValue();
-		if (!ResultValue.IsValid()
-			|| !FJsonSerializer::Serialize(
-				ResultValue.ToSharedRef(),
-				TJsonWriterFactory<TCHAR, TCondensedJsonPrintPolicy<TCHAR>>::Create(&OutResult)))
+		if (!ResultValue.IsValid())
 		{
 			OutError = FString::Printf(TEXT("Unable to serialize result from executing tool %s"), *ToolName);
 			return;
 		}
+		OutResult = UE::ToolsetRegistry::Internal::JsonToString(ResultValue.ToSharedRef());
 	}
 
 	TArray<FString> FToolsetLibraryWrapper::GetToolNames()
