Address post-commit feedback on 50900624.

#Synced-CL 50942406

diff --git a/Engine/Plugins/Experimental/ToolsetRegistry/Source/ToolsetRegistry/Private/ToolsetJson.cpp b/Engine/Plugins/Experimental/ToolsetRegistry/Source/ToolsetRegistry/Private/ToolsetJson.cpp
--- a/Engine/Plugins/Experimental/ToolsetRegistry/Source/ToolsetRegistry/Private/ToolsetJson.cpp
+++ b/Engine/Plugins/Experimental/ToolsetRegistry/Source/ToolsetRegistry/Private/ToolsetJson.cpp
@@ -56,6 +56,9 @@ const FStructProperty* AsColorProperty(TNotNull<const FProperty*> Property)
 	return nullptr;
 }
 
+// This method generates customizes the default JSON schema generation logic in a way that is easier for an
+// LLM to use reliably. The primary motivation for this method is converting UObject* and UClass* from
+// "string" types to a distinct reference type.
 bool CustomJsonSchema(
 	const FProperty* Property, const FString& ParameterDefaultString, const TSharedRef<FJsonObject>& OutputSchema)
 {
@@ -72,18 +75,10 @@ bool CustomJsonSchema(
 				ClassType = ClassProperty->MetaClass;
 				IsClass = true;
 			}
-			else if (const FObjectProperty* ObjectProperty = CastField<FObjectProperty>(Property))
+			else if (const FObjectPropertyBase* ObjectProperty = CastField<FObjectPropertyBase>(Property))
 			{
 				ClassType = ObjectProperty->PropertyClass;
 			}
-			else if (const FSoftObjectProperty* SoftProperty = CastField<FSoftObjectProperty>(Property))
-			{
-				ClassType = SoftProperty->PropertyClass;
-			}
-			else if (const FWeakObjectProperty* WeakProperty = CastField<FWeakObjectProperty>(Property))
-			{
-				ClassType = WeakProperty->PropertyClass;
-			}
 			else if (const FStructProperty* StructProperty = CastField<FStructProperty>(Property))
 			{
 				const UScriptStruct* ScriptStruct = StructProperty->Struct;
@@ -122,6 +117,8 @@ bool CustomJsonSchema(
 	return false;
 }
 
+// This method customizes the default JSON data import logic to match the custom schema and
+// data export logic in this file.
 bool CustomImport(const TSharedPtr<FJsonValue>& JsonValue, FProperty* Property, void* Value)
 {	
 	if (IsReferenceLikeProperty(Property))
@@ -143,12 +140,14 @@ bool CustomImport(const TSharedPtr<FJsonValue>& JsonValue, FProperty* Property,
 	return false;
 }
 
+// This method cutomizes the default JSON data export behavior to align with the
+// custom schema generated by CustomJsonSchema().
 TSharedPtr<FJsonValue> CustomExport(FProperty* Property, const void* Value)
 {
 	if (IsReferenceLikeProperty(Property))
 	{
 		FToolsetReference Reference;
-		if (FObjectProperty* ObjectProperty = CastField<FObjectProperty>(Property))
+		if (FObjectPropertyBase* ObjectProperty = CastField<FObjectPropertyBase>(Property))
 		{
 			if (const UObject* Object = ObjectProperty->GetObjectPropertyValue(Value))
 			{
@@ -158,29 +157,7 @@ TSharedPtr<FJsonValue> CustomExport(FProperty* Property, const void* Value)
 			{
 				return nullptr;
 			}
-		}
-		else if (FSoftObjectProperty* SoftObjectProperty = CastField<FSoftObjectProperty>(Property))
-		{
-			if (const UObject* Object = SoftObjectProperty->GetObjectPropertyValue(Value))
-			{
-				Reference.RefPath = Object->GetPathName();
-			}
-			else
-			{
-				return nullptr;
-			}
-		}
-		else if (FWeakObjectProperty* WeakObjectProperty = CastField<FWeakObjectProperty>(Property))
-		{
-			if (const UObject* Object = WeakObjectProperty->GetObjectPropertyValue(Value))
-			{
-				Reference.RefPath = Object->GetPathName();
-			}
-			else
-			{
-				return nullptr;
-			}
-		}
+		}		
 		else if (const FStructProperty* StructProperty = CastField<FStructProperty>(Property))
 		{
 			const UScriptStruct* ScriptStruct = StructProperty->Struct;
