[AI Assistant] Add *very* basic JSON schema builder to toolset registry.

Provides a way to build simple JSON schemas. This is intended to be used by
code that builds / modifies schemas and tests that validate generated schemas.

#Synced-CL 50976904

diff --git a/Engine/Plugins/Experimental/ToolsetRegistry/Source/ToolsetRegistry/Private/Tests/JsonSchemaTest.cpp b/Engine/Plugins/Experimental/ToolsetRegistry/Source/ToolsetRegistry/Private/Tests/JsonSchemaTest.cpp
new file mode 100644
--- /dev/null
+++ b/Engine/Plugins/Experimental/ToolsetRegistry/Source/ToolsetRegistry/Private/Tests/JsonSchemaTest.cpp
@@ -0,0 +1,107 @@
+// Copyright Epic Games, Inc. All Rights Reserved.
+
+#include "Containers/UnrealString.h"
+#include "Dom/JsonObject.h"
+#include "Misc/AutomationTest.h"
+#include "Templates/SharedPointer.h"
+
+#include "ToolsetRegistry/JsonConversion.h"
+#include "ToolsetRegistry/JsonSchema.h"
+
+#if WITH_DEV_AUTOMATION_TESTS
+
+BEGIN_DEFINE_SPEC(
+	FToolsetRegistryJsonSchemaSpec,
+	"AI.ToolsetRegistry.JsonSchema",
+	EAutomationTestFlags::EditorContext | EAutomationTestFlags::EngineFilter)
+END_DEFINE_SPEC(FToolsetRegistryJsonSchemaSpec)
+
+namespace
+{
+	const FString ExpectedDescription = TEXT("a description");
+}  // namespace
+
+void FToolsetRegistryJsonSchemaSpec::Define()
+{
+	using namespace UE::ToolsetRegistry::Internal;
+
+	Describe(TEXT("CreateJsonSchema"), [this]()
+	{
+		It("can create a schema for an array with no constraints", [this]()
+		{
+			TestEqual(
+				TEXT("JSON schema"),
+				JsonToString(CreateJsonSchema(ExpectedDescription, EJson::Array)),
+				FString::Printf(
+					TEXT(R"json({"description":"%s","type":"array"})json"),
+					*ExpectedDescription));
+		});
+
+		It("can create a schema for a boolean", [this]()
+		{
+			TestEqual(
+				TEXT("JSON schema"),
+				JsonToString(CreateJsonSchema(ExpectedDescription, EJson::Boolean)),
+				FString::Printf(
+					TEXT(R"json({"description":"%s","type":"boolean"})json"),
+					*ExpectedDescription));
+		});
+
+		It("can create a schema for null", [this]()
+		{
+			TestEqual(
+				TEXT("JSON schema"),
+				JsonToString(CreateJsonSchema(ExpectedDescription, EJson::Null)),
+				FString::Printf(
+					TEXT(R"json({"description":"%s","type":"null"})json"),
+					*ExpectedDescription));
+		});
+
+		It("can create a schema for a number", [this]()
+		{
+			TestEqual(
+				TEXT("JSON schema"),
+				JsonToString(CreateJsonSchema(ExpectedDescription, EJson::Number)),
+				FString::Printf(
+					TEXT(R"json({"description":"%s","type":"number"})json"),
+					*ExpectedDescription));
+		});
+
+		It("can create a schema for a string", [this]()
+		{
+			TestEqual(
+				TEXT("JSON schema"),
+				JsonToString(CreateJsonSchema(ExpectedDescription, EJson::String)),
+				FString::Printf(
+					TEXT(R"json({"description":"%s","type":"string"})json"),
+					*ExpectedDescription));
+		});
+
+		It("can create a schema for an empty object", [this]()
+		{
+			TestEqual(
+				TEXT("JSON schema"),
+				JsonToString(CreateJsonSchema(ExpectedDescription, EJson::Object)),
+				FString::Printf(
+					TEXT(R"json({"description":"%s","type":"object"})json"),
+					*ExpectedDescription));
+		});
+
+		It("can create a schema for an object", [this]()
+		{
+			const FString SpellDescription = TEXT("a spell to cast");
+			TSharedRef<FJsonObject> SpellSchema =
+				CreateJsonSchema(SpellDescription, EJson::String);
+			TSharedRef<FJsonObject> Properties = MakeShared<FJsonObject>();
+			Properties->SetObjectField(TEXT("spell"), SpellSchema);
+			TestEqual(
+				TEXT("JSON schema"),
+				JsonToString(CreateJsonSchema(ExpectedDescription, EJson::Object, Properties)),
+				FString::Printf(
+					TEXT(R"json({"description":"%s","type":"object","properties":%s})json"),
+					*ExpectedDescription, *JsonToString(Properties)));
+		});
+	});
+}
+
+#endif  // WITH_DEV_AUTOMATION_TESTS
diff --git a/Engine/Plugins/Experimental/ToolsetRegistry/Source/ToolsetRegistry/Private/ToolsetRegistry/JsonSchema.cpp b/Engine/Plugins/Experimental/ToolsetRegistry/Source/ToolsetRegistry/Private/ToolsetRegistry/JsonSchema.cpp
new file mode 100644
--- /dev/null
+++ b/Engine/Plugins/Experimental/ToolsetRegistry/Source/ToolsetRegistry/Private/ToolsetRegistry/JsonSchema.cpp
@@ -0,0 +1,65 @@
+// Copyright Epic Games, Inc. All Rights Reserved.
+
+#include "ToolsetRegistry/JsonSchema.h"
+
+#include "Containers/UnrealString.h"
+#include "Misc/AssertionMacros.h"
+#include "Templates/UnrealTemplate.h"
+
+namespace UE::ToolsetRegistry::Internal
+{
+	namespace
+	{
+		const FString JsonSchemaDescriptionFieldName = TEXT("description");
+		const FString JsonSchemaTypeFieldName = TEXT("type");
+		const FString JsonSchemaPropertiesFieldName = TEXT("properties");
+
+		// Maps JSON library types to valid JSON schema types.
+		const struct {
+			EJson Type;
+			FString Name;
+		} JsonSchemaTypeNames[] = {
+			{EJson::Array, TEXT("array")},
+			{EJson::Boolean, TEXT("boolean")},
+			{EJson::Null, TEXT("null")},
+			{EJson::Number, TEXT("number")},
+			{EJson::Object, TEXT("object")},
+			{EJson::String, TEXT("string")},
+		};
+
+		// Find a JSON schema type name given a JSON library type.
+		FString FindJsonSchemaTypeName(EJson Type)
+		{
+			check(Type != EJson::None);  // EJson::None is not supported.
+			FString Name;
+			for (int I = 0; I < UE_ARRAY_COUNT(JsonSchemaTypeNames); ++I)
+			{
+				const auto& JsonSchemaTypeName = JsonSchemaTypeNames[I];
+				if (Type == JsonSchemaTypeName.Type)
+				{
+					Name = JsonSchemaTypeName.Name;
+					break;
+				}
+			}
+			return Name;
+		}
+	}
+
+	TSharedRef<FJsonObject> CreateJsonSchema(
+		const FString& Description, EJson Type,
+		const TSharedPtr<FJsonObject> Properties)
+	{
+		FString TypeName = FindJsonSchemaTypeName(Type);
+		check(!TypeName.IsEmpty());
+		check(!Properties || Type == EJson::Object);
+		TSharedRef<FJsonObject> Schema = MakeShared<FJsonObject>();
+		Schema->SetStringField(JsonSchemaDescriptionFieldName, Description);
+		Schema->SetStringField(JsonSchemaTypeFieldName, TypeName);
+		if (Properties)
+		{
+			Schema->SetObjectField(JsonSchemaPropertiesFieldName, Properties);
+		}
+		return Schema;
+	}
+
+}  // namespace UE::ToolsetRegistry::Internal
\ No newline at end of file
diff --git a/Engine/Plugins/Experimental/ToolsetRegistry/Source/ToolsetRegistry/Private/ToolsetRegistry/JsonSchema.h b/Engine/Plugins/Experimental/ToolsetRegistry/Source/ToolsetRegistry/Private/ToolsetRegistry/JsonSchema.h
new file mode 100644
--- /dev/null
+++ b/Engine/Plugins/Experimental/ToolsetRegistry/Source/ToolsetRegistry/Private/ToolsetRegistry/JsonSchema.h
@@ -0,0 +1,32 @@
+// Copyright Epic Games, Inc. All Rights Reserved.
+
+#pragma once
+
+#include "Containers/UnrealString.h"
+#include "Dom/JsonObject.h"
+#include "Serialization/JsonTypes.h"
+#include "Templates/SharedPointer.h"
+
+namespace UE::ToolsetRegistry::Internal
+{
+	// Create a minimal JSON schema.
+	//
+	// See https://json-schema.org/
+	//
+	// This is intended to be used to create very simple schemas, if we need something that
+	// supports building arbitrary schemas we should probably turn this into a object that
+	// implements a builder pattern.
+	//
+	// Users still need to create constraints like:
+	// * Set the "items" or "prefixItems" field for arrays.
+	// * Manually create the properties object to define object fields.
+	//
+	// @param Description is the description of the schema.
+	// @param Type is the schema's type.
+	// @param Properties specifies the schemas for object fields when Type is EJson::Object.
+	//
+	// @returns JSON schema object.
+	TSharedRef<FJsonObject> CreateJsonSchema(
+		const FString& Description, EJson Type,
+		const TSharedPtr<FJsonObject> Properties = TSharedPtr<FJsonObject>());
+}  // namespace UE::ToolsetRegistry::Internal
\ No newline at end of file
