[AI Assistant] Add class used to provide JSON values or errors.

FJsonValueOrError provides a way to encode a JSON value and error within
a JSON object that - by default - adheres to the TypeScript style schema:

```
{"returnValue"?: unknown, "error"?: string}
```

FJsonValueOrError will be used to generate the final output schema for tools
and wrap their return values before they're sent to the assistant.

#Synced-CL 50981599

diff --git a/Engine/Plugins/Experimental/ToolsetRegistry/Source/ToolsetRegistry/Private/Tests/JsonValueOrErrorTest.cpp b/Engine/Plugins/Experimental/ToolsetRegistry/Source/ToolsetRegistry/Private/Tests/JsonValueOrErrorTest.cpp
new file mode 100644
--- /dev/null
+++ b/Engine/Plugins/Experimental/ToolsetRegistry/Source/ToolsetRegistry/Private/Tests/JsonValueOrErrorTest.cpp
@@ -0,0 +1,174 @@
+// Copyright Epic Games, Inc. All Rights Reserved.
+
+#include "Containers/UnrealString.h"
+#include "Dom/JsonObject.h"
+#include "Dom/JsonValue.h"
+#include "Misc/AutomationTest.h"
+#include "Templates/SharedPointer.h"
+
+#include "ToolsetRegistry/JsonConversion.h"
+#include "ToolsetRegistry/JsonValueOrError.h"
+#include "ToolsetRegistry/JsonSchema.h"
+
+#if WITH_DEV_AUTOMATION_TESTS
+
+BEGIN_DEFINE_SPEC(
+	FToolsetRegistryJsonValueOrErrorSpec,
+	"AI.ToolsetRegistry.JsonValueOrError",
+	EAutomationTestFlags::EditorContext | EAutomationTestFlags::EngineFilter)
+END_DEFINE_SPEC(FToolsetRegistryJsonValueOrErrorSpec)
+
+void FToolsetRegistryJsonValueOrErrorSpec::Define()
+{
+	using namespace UE::ToolsetRegistry;
+	using namespace UE::ToolsetRegistry::Internal;
+
+	Describe(TEXT("ToJsonObject"), [this]()
+	{
+		It(TEXT("should return a value if set"), [this]()
+		{
+			TSharedRef<FJsonValue> ExpectedValue = MakeShared<FJsonValueString>(TEXT("yo"));
+			FJsonValueOrError ValueOrError = MakeValue(ExpectedValue);
+			TSharedRef<FJsonObject> JsonObject = ValueOrError.ToJsonObject();
+			TestEqual(
+				TEXT("JSON object"),
+				JsonToString(JsonObject),
+				FString::Printf(TEXT(R"json({"%s":%s})json"),
+					*FJsonValueOrError::DefaultDescriptor.ValueFieldName,
+					*JsonToString(ExpectedValue)));
+		});
+
+		It(TEXT("should return null if a null value is set"), [this]()
+		{
+			FJsonValueOrError ValueOrError = MakeValue(TSharedPtr<FJsonValue>());
+			TSharedRef<FJsonObject> JsonObject = ValueOrError.ToJsonObject();
+			TestEqual(
+				TEXT("JSON object"),
+				JsonToString(JsonObject),
+				FString::Printf(TEXT(R"json({"%s":null})json"),
+					*FJsonValueOrError::DefaultDescriptor.ValueFieldName));
+		});
+
+		It(TEXT("should return an error if set"), [this]()
+		{
+			FString ExpectedError = TEXT("Failed");
+			FJsonValueOrError ValueOrError = MakeError(ExpectedError);
+			TSharedRef<FJsonObject> JsonObject = ValueOrError.ToJsonObject();
+			TestEqual(
+				TEXT("JSON object"),
+				JsonToString(JsonObject),
+				FString::Printf(TEXT(R"json({"%s":"%s"})json"),
+					*FJsonValueOrError::DefaultDescriptor.ErrorFieldName,
+					*ExpectedError));
+		});
+	});
+
+	Describe(TEXT("ToJsonString"), [this]()
+	{
+		It(TEXT("should return a value if set"), [this]()
+		{
+			FJsonValueOrError ValueOrError = MakeValue(MakeShared<FJsonValueString>(TEXT("yo")));
+			TestEqual(
+				TEXT("JSON object"),
+				JsonToString(ValueOrError.ToJsonObject()),
+				ValueOrError.ToJsonString());
+		});
+
+		It(TEXT("should return an error if set"), [this]()
+		{
+			FJsonValueOrError ValueOrError = MakeError(TEXT("Failed"));
+			TestEqual(
+				TEXT("JSON object"),
+				JsonToString(ValueOrError.ToJsonObject()),
+				ValueOrError.ToJsonString());
+		});
+	});
+
+	Describe(TEXT("FromJsonStringOrError"), [this]()
+	{
+		It(TEXT("should return a value if set"), [this]()
+		{
+			FString ExpectedValueString = TEXT(R"json("yo")json");
+			FJsonValueOrError ValueOrError =
+				FJsonValueOrError::FromJsonStringOrError(MakeValue(ExpectedValueString));
+			TestEqual(
+				TEXT("JSON object"),
+				JsonToString(ValueOrError.ToJsonObject()),
+				FString::Printf(
+					TEXT(R"json({"%s":%s})json"),
+					*FJsonValueOrError::DefaultDescriptor.ValueFieldName,
+					*ExpectedValueString));
+		});
+
+
+		It(TEXT("should return an error if set"), [this]()
+		{
+			FString ExpectedError = TEXT("Failed");
+			FJsonValueOrError ValueOrError =
+				FJsonValueOrError::FromJsonStringOrError(MakeError(ExpectedError));
+			TestEqual(
+				TEXT("JSON object"),
+				JsonToString(ValueOrError.ToJsonObject()),
+				FString::Printf(
+					TEXT(R"json({"%s":"%s"})json"),
+					*FJsonValueOrError::DefaultDescriptor.ErrorFieldName,
+					*ExpectedError));
+		});
+		
+		It(TEXT("should fail if provided an invalid JSON string"), [this]()
+		{
+			FString NotJson = TEXT("not_json");
+			FJsonValueOrError ValueOrError =
+				FJsonValueOrError::FromJsonStringOrError(MakeValue(NotJson));
+			TestEqual(
+				TEXT("JSON object"),
+				JsonToString(ValueOrError.ToJsonObject()),
+				FString::Printf(
+					TEXT(R"json({"%s":"Failed to parse JSON result '%s'"})json"),
+					*FJsonValueOrError::DefaultDescriptor.ErrorFieldName,
+					*NotJson));
+		});
+	});
+
+	Describe(TEXT("GetJsonSchema"), [this]()
+	{
+		It(TEXT("should return a value / error schema from an object schema"), [this]()
+		{
+			TSharedRef<FJsonObject> NestedSchema = CreateJsonSchema(TEXT("magic"), EJson::String);
+			TSharedRef<FJsonObject> Schema = FJsonValueOrError::GetJsonSchema(NestedSchema);
+			TestEqual(
+				TEXT("JSON schema"),
+				JsonToString(Schema),
+				FString::Printf(
+					TEXT(
+						R"json({"description":"%s","type":"object","properties":{)json"
+						R"json("%s":%s,"%s":{"description":"%s","type":"string"}}})json"),
+					*FJsonValueOrError::DefaultDescriptor.ObjectDescription,
+					*FJsonValueOrError::DefaultDescriptor.ValueFieldName,
+					*JsonToString(NestedSchema),
+					*FJsonValueOrError::DefaultDescriptor.ErrorFieldName,
+					*FJsonValueOrError::DefaultDescriptor.ErrorFieldDescription));
+		});
+
+		It(TEXT("should return a value / error schema with a null value from no schema"), [this]()
+		{
+			TSharedRef<FJsonObject> Schema =
+				FJsonValueOrError::GetJsonSchema(TSharedPtr<FJsonObject>());
+			TestEqual(
+				TEXT("JSON schema"),
+				JsonToString(Schema),
+				FString::Printf(
+					TEXT(
+						R"json({"description":"%s","type":"object","properties":{)json"
+						R"json("%s":{"description":"%s","type":"null"},)json"
+						R"json("%s":{"description":"%s","type":"string"}}})json"),
+					*FJsonValueOrError::DefaultDescriptor.ObjectDescription,
+					*FJsonValueOrError::DefaultDescriptor.ValueFieldName,
+					*FJsonValueOrError::DefaultDescriptor.NullValueFieldDescription,
+					*FJsonValueOrError::DefaultDescriptor.ErrorFieldName,
+					*FJsonValueOrError::DefaultDescriptor.ErrorFieldDescription));
+		});
+	});
+}
+
+#endif  // WITH_DEV_AUTOMATION_TESTS
\ No newline at end of file
diff --git a/Engine/Plugins/Experimental/ToolsetRegistry/Source/ToolsetRegistry/Private/ToolsetRegistry/JsonValueOrError.cpp b/Engine/Plugins/Experimental/ToolsetRegistry/Source/ToolsetRegistry/Private/ToolsetRegistry/JsonValueOrError.cpp
new file mode 100644
--- /dev/null
+++ b/Engine/Plugins/Experimental/ToolsetRegistry/Source/ToolsetRegistry/Private/ToolsetRegistry/JsonValueOrError.cpp
@@ -0,0 +1,93 @@
+// Copyright Epic Games, Inc. All Rights Reserved.
+
+#include "ToolsetRegistry/JsonValueOrError.h"
+
+#include "Containers/UnrealString.h"
+#include "Misc/AssertionMacros.h"
+#include "Templates/Function.h"
+
+#include "ToolsetRegistry/JsonConversion.h"
+#include "ToolsetRegistry/JsonSchema.h"
+
+namespace UE::ToolsetRegistry
+{
+	// Default descriptor for a FJsonValueOrError.
+	const FJsonValueOrErrorSchemaDescriptor FJsonValueOrError::DefaultDescriptor = {
+		TEXT("Either a value if successful or an error if the operation failed."),
+		TEXT("returnValue"),
+		TEXT("Always set to null on success."),
+		TEXT("error"),
+		TEXT("If set, describes the error that occurred."),
+	};
+
+	TSharedRef<FJsonObject> FJsonValueOrError::ToJsonObject(
+		const FJsonValueOrErrorSchemaDescriptor& Descriptor) const
+	{
+		TSharedRef<FJsonObject> JsonObject = MakeShared<FJsonObject>();
+		if (HasError())
+		{
+			JsonObject->SetStringField(Descriptor.ErrorFieldName, GetError());
+		}
+		else
+		{
+			// If there is no error, this object must have a value.
+			check(HasValue());
+			TSharedPtr<FJsonValue> JsonValue = GetValue();
+			// If there is no value return null.
+			JsonObject->SetField(
+				Descriptor.ValueFieldName,
+				JsonValue ? JsonValue : MakeShared<FJsonValueNull>());
+		}
+		return JsonObject;
+	}
+
+	FString FJsonValueOrError::ToJsonString(
+		const FJsonValueOrErrorSchemaDescriptor& Descriptor) const
+	{
+		return UE::ToolsetRegistry::Internal::JsonToString(ToJsonObject(Descriptor));
+	}
+
+	TSharedRef<FJsonObject> FJsonValueOrError::GetJsonSchema(
+		TSharedPtr<FJsonObject> ValueSchema,
+		const FJsonValueOrErrorSchemaDescriptor& Descriptor)
+	{
+		// Map the value and error field names to their schemas.
+		TSharedRef<FJsonObject> Properties = MakeShared<FJsonObject>();
+		Properties->SetObjectField(
+			Descriptor.ValueFieldName,
+			// Use "null" as the value's type if it has no schema.
+			ValueSchema
+			? ValueSchema
+			: Internal::CreateJsonSchema(Descriptor.NullValueFieldDescription, EJson::Null));
+		Properties->SetObjectField(
+			Descriptor.ErrorFieldName,
+			Internal::CreateJsonSchema(
+				Descriptor.ErrorFieldDescription, EJson::String));
+
+		// Create the schema for the overall object.
+		TSharedRef<FJsonObject> Schema =
+			Internal::CreateJsonSchema(
+				Descriptor.ObjectDescription, EJson::Object, Properties);
+		return Schema;
+	}
+
+	FJsonValueOrError FJsonValueOrError::FromJsonStringOrError(
+		const TValueOrError<FString, FString>& JsonStringOrError)
+	{
+		if (JsonStringOrError.HasError())
+		{
+			return MakeError(JsonStringOrError.GetError());
+		}
+		check(JsonStringOrError.HasValue());
+		const FString& JsonValueString = JsonStringOrError.GetValue();
+		TSharedPtr<FJsonValue> JsonValue =
+			UE::ToolsetRegistry::Internal::JsonStringToJsonValue(JsonValueString);
+		if (!JsonValue)
+		{
+			return MakeError(
+				FString::Printf(TEXT("Failed to parse JSON result '%s'"), *JsonValueString));
+		}
+		return MakeValue(JsonValue); 
+	}
+
+}  // namespace UE::ToolsetRegistry
\ No newline at end of file
diff --git a/Engine/Plugins/Experimental/ToolsetRegistry/Source/ToolsetRegistry/Public/ToolsetRegistry/JsonValueOrError.h b/Engine/Plugins/Experimental/ToolsetRegistry/Source/ToolsetRegistry/Public/ToolsetRegistry/JsonValueOrError.h
new file mode 100644
--- /dev/null
+++ b/Engine/Plugins/Experimental/ToolsetRegistry/Source/ToolsetRegistry/Public/ToolsetRegistry/JsonValueOrError.h
@@ -0,0 +1,87 @@
+// Copyright Epic Games, Inc. All Rights Reserved.
+
+#pragma once
+
+#include "Containers/UnrealString.h"
+#include "Dom/JsonObject.h"
+#include "Dom/JsonValue.h"
+#include "Templates/SharedPointer.h"
+#include "Templates/ValueOrError.h"
+
+#define UE_API TOOLSETREGISTRY_API
+
+namespace UE::ToolsetRegistry
+{
+	/// Base class for FValueOrError.
+	using FJsonValueOrErrorBase = TValueOrError<TSharedPtr<FJsonValue>, FString>;
+
+	/// Schema descriptor for a JSON value or error.
+	struct FJsonValueOrErrorSchemaDescriptor
+	{
+		/// Top level description for a value or error object.
+		const FString ObjectDescription;
+		/// Value field name in an object of this type.
+		const FString ValueFieldName;
+		/// Description of the value field if it can only be null.
+		const FString NullValueFieldDescription;
+		/// Error field name in an object of this type.
+		const FString ErrorFieldName;
+		/// Description of the error field.
+		const FString ErrorFieldDescription;
+	};
+
+	/// JSON value or error that can be converted to / from a JSON object that is compliant with 
+	/// the schema.
+	/// 
+	/// ```typescript
+	/// { "returnValue"?: unknown, "error"?: string }
+	/// ```
+	class FJsonValueOrError : public FJsonValueOrErrorBase
+	{
+	public:
+		using FJsonValueOrErrorBase::FJsonValueOrErrorBase;
+
+		/// Convert to a value / error JSON object.
+		///
+		/// @param Descriptor Describes how to map value and error fields to the returned object.
+		///
+		/// @returns JSON object that contains the value or error.
+		UE_API TSharedRef<FJsonObject> ToJsonObject(
+			const FJsonValueOrErrorSchemaDescriptor& Descriptor = DefaultDescriptor) const;
+
+		/// Convert to a JSON string.
+		///
+		/// @param Descriptor Describes how to map value and error fields to the returned object.
+		///
+		/// @returns JSON string that contains the value or error.
+		UE_API FString ToJsonString(
+			const FJsonValueOrErrorSchemaDescriptor& Descriptor = DefaultDescriptor) const;
+
+	public:
+		// Default descriptor for a FJsonValueOrError.
+		UE_API static const FJsonValueOrErrorSchemaDescriptor DefaultDescriptor;
+
+		/// Convert JSON schema for a JSON value / object into one that is wrapped
+		/// by an object of this type (i.e a value or error).
+		///
+		/// @param ValueSchema Schema of the value stored by an instance of this object.
+		/// @param Descriptor Describes how to map value and error fields to the returned object.
+		///
+		/// @returns JSON schema for a FJsonValueOrError that has ValueSchema as the value type.
+		UE_API static TSharedRef<FJsonObject> GetJsonSchema(
+			TSharedPtr<FJsonObject> ValueSchema,
+			const FJsonValueOrErrorSchemaDescriptor& Descriptor = DefaultDescriptor);
+
+		/// Construct from JSON string or error. If a JSON string cannot be parsed this returns
+		/// an error.
+		///
+		/// @param JsonStringOrError JSON string or error to convert to an instance of this object.
+		///
+		/// @returns Instance of FJsonValueOrError that contains either the JSON value or error
+		/// from JsonStringOrError.
+		UE_API static FJsonValueOrError FromJsonStringOrError(
+			const TValueOrError<FString, FString>& JsonStringOrError);
+	};
+}  // namespace UE::ToolsetRegistry
+
+#undef UE_API
\ No newline at end of file
