[AI Assistant] Fix up file names and clean up some toolset registry naming.

* Renamed FToolsetLibraryWrapper to BlueprintLibraryToolset and moved into
files that follow that naming.
* Moved FObjectFunctionToolCall into files that use the appropriate names.
* Moved JsonUtilitiesEditor module loading into ToolsetJson since that is now
the C++ module that depends upon JsonUtilitiesEditor.
* Moved ToolsetJson methods under an internal namespace.

#Synced-CL 50992476

diff --git a/Engine/Plugins/Experimental/ToolsetRegistry/Source/ToolsetRegistry/Private/AgentSkill.cpp b/Engine/Plugins/Experimental/ToolsetRegistry/Source/ToolsetRegistry/Private/AgentSkill.cpp
--- a/Engine/Plugins/Experimental/ToolsetRegistry/Source/ToolsetRegistry/Private/AgentSkill.cpp
+++ b/Engine/Plugins/Experimental/ToolsetRegistry/Source/ToolsetRegistry/Private/AgentSkill.cpp
@@ -6,7 +6,7 @@
 #include "AssetRegistry/AssetData.h"
 #include "Interfaces/ITargetPlatform.h"
 
-#include "ToolsetRegistryToolset.h"
+#include "ToolsetRegistry/BlueprintLibraryToolset.h"
 #include "ToolsetRegistry/Toolset.h"
 
 #include UE_INLINE_GENERATED_CPP_BY_NAME(AgentSkill)
@@ -18,7 +18,8 @@ FAgentSkillDetails UAgentSkill::GetDetails() const
 	for (const auto& Toolset : Toolsets)
 	{
 		UClass* ToolsetClass = Toolset.Get();
-		FString ToolsetName = UE::ToolsetRegistry::FToolsetLibraryWrapper::GetToolsetClassName(ToolsetClass);
+		FString ToolsetName =
+			UE::ToolsetRegistry::FBlueprintLibraryToolset::GetToolsetClassName(ToolsetClass);
 		SkillDetails.Toolsets.Add(ToolsetName);
 	}
 	return SkillDetails;
diff --git a/Engine/Plugins/Experimental/ToolsetRegistry/Source/ToolsetRegistry/Private/BPToolsetRegistry.cpp b/Engine/Plugins/Experimental/ToolsetRegistry/Source/ToolsetRegistry/Private/BPToolsetRegistry.cpp
--- a/Engine/Plugins/Experimental/ToolsetRegistry/Source/ToolsetRegistry/Private/BPToolsetRegistry.cpp
+++ b/Engine/Plugins/Experimental/ToolsetRegistry/Source/ToolsetRegistry/Private/BPToolsetRegistry.cpp
@@ -7,10 +7,10 @@
 #include "Misc/Variant.h"
 #include "Templates/SharedPointer.h"
 
-#include "ToolsetRegistryToolset.h"
 #include "ToolsetRegistry.h"
 #include "ToolsetRegistryModule.h"
 #include "ToolsetRegistrySubsystem.h"
+#include "ToolsetRegistry/BlueprintLibraryToolset.h"
 #include "ToolsetRegistry/Toolset.h"
 
 void UToolsetRegistry::RegisterToolsetClass(TSubclassOf<UToolsetDefinition> InToolsetClass)
@@ -21,27 +21,31 @@ void UToolsetRegistry::RegisterToolsetClass(TSubclassOf<UToolsetDefinition> InTo
 		return;
 	}
 
-	if (auto ToolsetRegistrySubsystem = UToolsetRegistrySubsystem::Get(TEXT("IsToolsetClassRegistered"));
+	if (auto ToolsetRegistrySubsystem =
+			UToolsetRegistrySubsystem::Get(TEXT("IsToolsetClassRegistered"));
 		ToolsetRegistrySubsystem.HasValue())
 	{
-		FString ToolsetName = UE::ToolsetRegistry::FToolsetLibraryWrapper::GetToolsetClassName(InToolsetClass);
-		TSharedPtr<UE::ToolsetRegistry::FToolsetLibraryWrapper> Wrapper =
-			MakeShared<UE::ToolsetRegistry::FToolsetLibraryWrapper>(InToolsetClass);
+		using UE::ToolsetRegistry::FBlueprintLibraryToolset;
+		TSharedPtr<FBlueprintLibraryToolset> Toolset =
+			MakeShared<FBlueprintLibraryToolset>(InToolsetClass);
 
-		if (!Wrapper->HasValidTools())
+		FString ToolsetName = FBlueprintLibraryToolset::GetToolsetClassName(InToolsetClass);
+		if (!Toolset->HasValidTools())
 		{
 			UE_LOG(LogToolsetRegistry, Warning,
 				   TEXT("Unable to register Toolset %s: invalid tool list"), *ToolsetName);
 			return;
 		}
 
-		if (ToolsetRegistrySubsystem.GetValue()->ToolsetRegistry.RegisterToolset(Wrapper))
+		if (ToolsetRegistrySubsystem.GetValue()->ToolsetRegistry.RegisterToolset(Toolset))
 		{
-			UE_LOG(LogToolsetRegistry, Display, TEXT("Registering Toolset %s"), *ToolsetName);
+			UE_LOG(LogToolsetRegistry, Display,
+				TEXT("Registering Toolset %s"), *ToolsetName);
 		}
 		else
 		{
-			UE_LOG(LogToolsetRegistry, Warning, TEXT("Unable to register Toolset %s"), *ToolsetName);
+			UE_LOG(LogToolsetRegistry, Warning,
+				TEXT("Unable to register Toolset %s"), *ToolsetName);
 		}
 	}
 	else
@@ -62,7 +66,8 @@ void UToolsetRegistry::UnregisterToolsetClass(TSubclassOf<UToolsetDefinition> In
 			UToolsetRegistrySubsystem::Get(TEXT("IsToolsetClassRegistered"));
 		ToolsetRegistrySubsystem.HasValue())
 	{
-		FString ToolsetName = UE::ToolsetRegistry::FToolsetLibraryWrapper::GetToolsetClassName(InToolsetClass);
+		FString ToolsetName =
+			UE::ToolsetRegistry::FBlueprintLibraryToolset::GetToolsetClassName(InToolsetClass);
 		TSharedPtr<UE::ToolsetRegistry::FToolset> Handler =
 			ToolsetRegistrySubsystem.GetValue()->ToolsetRegistry.Find(ToolsetName);
 		if (Handler)
@@ -87,7 +92,8 @@ bool UToolsetRegistry::IsToolsetClassRegistered(TSubclassOf<UToolsetDefinition>
 	if (auto ToolsetRegistrySubsystem = UToolsetRegistrySubsystem::Get(TEXT("IsToolsetClassRegistered"));
 		ToolsetRegistrySubsystem.HasValue())
 	{
-		FString ToolsetName = UE::ToolsetRegistry::FToolsetLibraryWrapper::GetToolsetClassName(InToolsetClass);
+		FString ToolsetName =
+			UE::ToolsetRegistry::FBlueprintLibraryToolset::GetToolsetClassName(InToolsetClass);
 		return ToolsetRegistrySubsystem.GetValue()->ToolsetRegistry.Find(ToolsetName) != nullptr;
 	}
 	return false;
@@ -144,7 +150,7 @@ FString UToolsetRegistry::GetToolsetJsonSchema(TSubclassOf<UToolsetDefinition> I
 		return FString();
 	}
 
-	TSharedPtr<UE::ToolsetRegistry::FToolsetLibraryWrapper> Wrapper =
-		MakeShared<UE::ToolsetRegistry::FToolsetLibraryWrapper>(InToolsetClass);
-	return Wrapper ? Wrapper->GetToolsetJsonSchema() : FString();
+	TSharedPtr<UE::ToolsetRegistry::FBlueprintLibraryToolset> Toolset =
+		MakeShared<UE::ToolsetRegistry::FBlueprintLibraryToolset>(InToolsetClass);
+	return Toolset ? Toolset->GetToolsetJsonSchema() : FString();
 }
diff --git a/Engine/Plugins/Experimental/ToolsetRegistry/Source/ToolsetRegistry/Private/Tests/AgentSkillsTest.cpp b/Engine/Plugins/Experimental/ToolsetRegistry/Source/ToolsetRegistry/Private/Tests/AgentSkillsTest.cpp
--- a/Engine/Plugins/Experimental/ToolsetRegistry/Source/ToolsetRegistry/Private/Tests/AgentSkillsTest.cpp
+++ b/Engine/Plugins/Experimental/ToolsetRegistry/Source/ToolsetRegistry/Private/Tests/AgentSkillsTest.cpp
@@ -18,12 +18,11 @@
 #include "UObject/ObjectMacros.h"
 #include "Tests/AgentSkillsTest.h"
 #include "Tests/ToolsetRegistryTestFlags.h"
-#include "ToolsetRegistryToolset.h"
+#include "ToolsetRegistry/BlueprintLibraryToolset.h"
 
 
 #if WITH_DEV_AUTOMATION_TESTS
 
-
 BEGIN_DEFINE_SPEC(FAgentSkillSpec, "AI.ToolsetRegistry.AgentSkillSpec",
 	EAutomationTestFlags::EditorContext | EAutomationTestFlags::EngineFilter)	
 	FString MountPoint;
@@ -73,20 +72,24 @@ void FAgentSkillSpec::Define()
 
 	It("Get skill details", [this]()
 	{
-		TMap<FString, FAgentSkillDetails> Summaries = UAgentSkillToolset::GetSkills(TArray<FString>({ PackagePath }));
+		TMap<FString, FAgentSkillDetails> Summaries =
+			UAgentSkillToolset::GetSkills(TArray<FString>({ PackagePath }));
 		FAgentSkillDetails* Details = Summaries.Find(PackagePath);
 		TestTrue("Details", Details != nullptr);
 		if (Details)
 		{
 			TestEqual("Details", Details->Instructions, AgentSkill->Instructions);
 			TestEqual("Details", Details->Toolsets.Num(), AgentSkill->Toolsets.Num());
-			TestEqual("Details", Details->Toolsets[0], UE::ToolsetRegistry::FToolsetLibraryWrapper::GetToolsetClassName(AgentSkill->Toolsets[0].Get()));
+			TestEqual("Details", Details->Toolsets[0],
+				UE::ToolsetRegistry::FBlueprintLibraryToolset::GetToolsetClassName(
+					AgentSkill->Toolsets[0].Get()));
 		}
 	});
 
 	It("Can't load non-existant skills", [this]()
 	{
-		TMap<FString, FAgentSkillDetails> Summaries = UAgentSkillToolset::GetSkills(TArray<FString>({ "DoesNotExist" }));		
+		TMap<FString, FAgentSkillDetails> Summaries =
+			UAgentSkillToolset::GetSkills(TArray<FString>({ "DoesNotExist" }));
 		TestEqual("Details", Summaries.Num(), 0);
 	});
 
@@ -95,7 +98,7 @@ void FAgentSkillSpec::Define()
 		TObjectPtr<UAgentSkill> AgentSkillCustom = NewObject<UAgentSkillCustomPrompt>();
 		AgentSkillCustom->Instructions = FString(TEXT("FooBar"));
 		FAgentSkillDetails Details = AgentSkillCustom->GetDetails();
-		TestEqual("Details", Details.Instructions, FString(TEXT("FooBarBishBosh")));		
+		TestEqual("Details", Details.Instructions, FString(TEXT("FooBarBishBosh")));
 	});
 }
 
diff --git a/Engine/Plugins/Experimental/ToolsetRegistry/Source/ToolsetRegistry/Private/Tests/BlueprintLibraryToolsetTest.cpp b/Engine/Plugins/Experimental/ToolsetRegistry/Source/ToolsetRegistry/Private/Tests/BlueprintLibraryToolsetTest.cpp
new file mode 100644
--- /dev/null
+++ b/Engine/Plugins/Experimental/ToolsetRegistry/Source/ToolsetRegistry/Private/Tests/BlueprintLibraryToolsetTest.cpp
@@ -0,0 +1,770 @@
+// Copyright Epic Games, Inc. All Rights Reserved.
+
+#include "BlueprintLibraryToolsetTest.h"
+
+#include "CoreMinimal.h"
+
+#include "Misc/AutomationTest.h"
+#include "Dom/JsonObject.h"
+#include "JsonDomBuilder.h"
+#include "IPythonScriptPlugin.h"
+#include "PythonScriptTypes.h"
+#include "Policies/CondensedJsonPrintPolicy.h"
+#include "Serialization/JsonReader.h"
+#include "Serialization/JsonWriter.h"
+#include "Serialization/JsonSerializer.h"
+
+#include "Tests/ToolsetRegistryTestSchemaBuilder.h"
+#include "Tests/ToolsetRegistryTestFlags.h"
+
+#include "BPToolsetRegistry.h"
+#include "ToolsetRegistrySubsystem.h"
+#include "ToolsetRegistry/BlueprintLibraryToolset.h"
+#include "ToolsetRegistry/JsonConversion.h"
+#include "ToolsetRegistry/Toolset.h"
+
+#if WITH_DEV_AUTOMATION_TESTS
+const FString FakeToolsetName(TEXT("ToolsetRegistry.FakeToolset"));
+const FString FakeToolsetInput(TEXT(R"({"inString":"hello"})"));
+const FString FakeToolName(TEXT("SimpleTool"));
+const FString ReturnValueName(TEXT("returnValue"));
+const FString ErrorValueName(TEXT("error"));
+
+class ToolsetRegistryContainer
+{
+  public:
+	ToolsetRegistryContainer(TSubclassOf<UToolsetDefinition> InToolsetClass)
+		: ToolsetClass(InToolsetClass)
+	{
+		UToolsetRegistry::RegisterToolsetClass(ToolsetClass);
+	}
+
+	~ToolsetRegistryContainer()
+	{
+		UToolsetRegistry::UnregisterToolsetClass(ToolsetClass);
+	}
+
+	TSharedPtr<UE::ToolsetRegistry::FBlueprintLibraryToolset> ToolsetHandler()
+	{
+		TSharedPtr<UE::ToolsetRegistry::FToolset> Handler;
+		auto ToolsetRegistrySubsystem = UToolsetRegistrySubsystem::Get(TEXT("Test Register"));
+		if (ToolsetRegistrySubsystem.HasValue() && ToolsetClass)
+		{
+			FString ToolsetName =
+				UE::ToolsetRegistry::FBlueprintLibraryToolset::GetToolsetClassName(ToolsetClass);
+			Handler = ToolsetRegistrySubsystem.GetValue()->ToolsetRegistry.Find(ToolsetName);
+		}
+		return StaticCastSharedPtr<UE::ToolsetRegistry::FBlueprintLibraryToolset>(Handler);
+	}
+
+	TSubclassOf<UToolsetDefinition> ToolsetClass;
+};
+
+TSharedPtr<FJsonObject> GetFakeToolsetExpectedSchema()
+{
+	FJsonDomBuilder::FObject SchemaBuilder = MakeSchemaBuilder(
+		FString(FakeToolsetName), FString("1.0"), FString(TEXT("Fake Toolset that does nothing.")));
+
+	FJsonDomBuilder::FArray ToolsArray;
+	// Add SimpleTool
+	{
+		FJsonDomBuilder::FObject InputSchema = MakeSchema();
+		FJsonDomBuilder::FObject Properties;
+
+		FString ParamName = TEXT("inString");
+		Properties.Set(ParamName, MakeProperty(TEXT("string")));
+		AddPropertiesToSchema(InputSchema, Properties);
+		AddRequiredListToSchema(InputSchema, TArray<FString>({ParamName}));
+
+		FJsonDomBuilder::FObject OutputSchema = MakeSchema();
+		FJsonDomBuilder::FObject OutputProperties;
+		OutputProperties.Set(ReturnValueName, MakeProperty(TEXT("string")));
+		AddPropertiesToSchema(OutputSchema, OutputProperties);
+		AddRequiredListToSchema(OutputSchema, TArray<FString>({ReturnValueName}));
+
+		FJsonDomBuilder::FObject ToolEntry = MakeToolSchemaEntry(
+			TEXT("ToolsetRegistry.FakeToolset.SimpleTool"), TEXT("Simple tool that does nothing."),
+			InputSchema, OutputSchema);
+
+		AddToolSchema(SchemaBuilder, ToolsArray, ToolEntry);
+	}
+
+	// Add SimpleNoReturnTool
+	{
+		FJsonDomBuilder::FObject InputSchema = MakeSchema();
+		FJsonDomBuilder::FObject Properties;
+
+		FString ParamName = TEXT("inString");
+		Properties.Set(ParamName, MakeProperty(TEXT("string")));
+		AddPropertiesToSchema(InputSchema, Properties);
+		AddRequiredListToSchema(InputSchema, TArray<FString>({ParamName}));
+
+		FJsonDomBuilder::FObject ToolEntry;
+		ToolEntry.Set(TEXT("description"), TEXT("Simple tool that does nothing and returns nothing."));
+		ToolEntry.Set(TEXT("name"), TEXT("ToolsetRegistry.FakeToolset.SimpleNoReturnTool"));
+		ToolEntry.Set(TEXT("inputSchema"), InputSchema);
+
+		AddToolSchema(SchemaBuilder, ToolsArray, ToolEntry);
+	}
+
+	return SchemaBuilder.AsJsonObject();
+}
+
+TSharedPtr<FJsonObject> GetFakeToolsetWithDefaultParamsExpectedSchema()
+{
+	FJsonDomBuilder::FObject SchemaBuilder = MakeSchemaBuilder(
+		FString(TEXT("ToolsetRegistry.FakeToolsetWithDefaultParams")), FString("1.0"),
+		FString(TEXT("Fake Toolset that has method with default parameters.")));
+
+	FJsonDomBuilder::FObject InputSchema = MakeSchema();
+	FJsonDomBuilder::FObject Properties;
+
+	FString ParamName1 = TEXT("inString");
+	FString ParamName2 = TEXT("inValue");
+
+	Properties.Set(ParamName1, MakeProperty(TEXT("string")));
+	Properties.Set(ParamName2, MakeProperty(TEXT("integer"), MakeShared<FJsonValueNumber>(42)));
+	AddPropertiesToSchema(InputSchema, Properties);
+	AddRequiredListToSchema(InputSchema, TArray<FString>({ParamName1}));
+
+	FJsonDomBuilder::FObject OutputSchema = MakeSchema();
+	FJsonDomBuilder::FObject OutputProperties;
+	OutputProperties.Set(ReturnValueName, MakeProperty(TEXT("string")));
+	AddPropertiesToSchema(OutputSchema, OutputProperties);
+	AddRequiredListToSchema(OutputSchema, TArray<FString>({ReturnValueName}));
+
+	FJsonDomBuilder::FObject ToolEntry = MakeToolSchemaEntry(
+		TEXT("ToolsetRegistry.FakeToolsetWithDefaultParams.ToolWithDefaults"),
+		TEXT("Method with some params with default values."),
+		InputSchema, OutputSchema);
+
+	FJsonDomBuilder::FArray ToolsArray;
+	AddToolSchema(SchemaBuilder, ToolsArray, ToolEntry);
+
+	return SchemaBuilder.AsJsonObject();
+}
+
+TSharedPtr<FJsonObject> GetFakeToolsetWithArrayParamsExpectedSchema()
+{
+	FJsonDomBuilder::FObject SchemaBuilder = MakeSchemaBuilder(
+		FString(TEXT("ToolsetRegistry.FakeToolsetWithArrayParams")), FString("1.0"),
+		FString(TEXT("Fake Toolset that has method with TArray parameters.")));
+
+	FJsonDomBuilder::FObject InputSchema = MakeSchema();
+
+	FString ParamName1 = TEXT("arrayParamOne");
+	FString ParamName2 = TEXT("arrayParamTwo");
+
+	FJsonDomBuilder::FObject InputProperties;
+	InputProperties.Set(ParamName1, MakeArrayProperty(TEXT("integer")));
+	InputProperties.Set(ParamName2, MakeArrayProperty(TEXT("string")));
+	AddPropertiesToSchema(InputSchema, InputProperties);
+	AddRequiredListToSchema(InputSchema, TArray<FString>({ParamName1, ParamName2}));
+
+	FJsonDomBuilder::FObject OutputSchema = MakeSchema();
+	FJsonDomBuilder::FObject OutputProperties;
+	OutputProperties.Set(ReturnValueName, MakeArrayProperty(TEXT("number")));
+	AddPropertiesToSchema(OutputSchema, OutputProperties);
+	AddRequiredListToSchema(OutputSchema, TArray<FString>({ReturnValueName}));
+
+	FJsonDomBuilder::FObject ToolEntry = MakeToolSchemaEntry(
+		TEXT("ToolsetRegistry.FakeToolsetWithArrayParams.ToolWithArrays"),
+		TEXT("Method with some params with TArray values."),
+		InputSchema, OutputSchema);
+
+	FJsonDomBuilder::FArray ToolsArray;
+	AddToolSchema(SchemaBuilder, ToolsArray, ToolEntry);
+
+	return SchemaBuilder.AsJsonObject();
+}
+
+TSharedPtr<FJsonObject> GetFakeToolsetWithSetParamsExpectedSchema()
+{
+	FJsonDomBuilder::FObject SchemaBuilder = MakeSchemaBuilder(
+		FString(TEXT("ToolsetRegistry.FakeToolsetWithSetParams")), FString("1.0"),
+		FString(TEXT("Fake Toolset that has method with TSet parameters.")));
+
+	FJsonDomBuilder::FObject InputSchema = MakeSchema();
+
+	FString ParamName1 = TEXT("setParamOne");
+	FString ParamName2 = TEXT("setParamTwo");
+
+	FJsonDomBuilder::FObject InputProperties;
+	InputProperties.Set(ParamName1, MakeSetProperty(TEXT("integer")));
+	InputProperties.Set(ParamName2, MakeSetProperty(TEXT("string")));
+	AddPropertiesToSchema(InputSchema, InputProperties);
+	AddRequiredListToSchema(InputSchema, TArray<FString>({ParamName1, ParamName2}));
+
+	FJsonDomBuilder::FObject OutputSchema = MakeSchema();
+	FJsonDomBuilder::FObject OutputProperties;
+	OutputProperties.Set(ReturnValueName, MakeSetProperty(TEXT("number")));
+	AddPropertiesToSchema(OutputSchema, OutputProperties);
+	AddRequiredListToSchema(OutputSchema, TArray<FString>({ReturnValueName}));
+
+	FJsonDomBuilder::FObject ToolEntry = MakeToolSchemaEntry(
+		TEXT("ToolsetRegistry.FakeToolsetWithSetParams.ToolWithSets"),
+		TEXT("Method with some params with TSet values."),
+		InputSchema, OutputSchema);
+
+	FJsonDomBuilder::FArray ToolsArray;
+	AddToolSchema(SchemaBuilder, ToolsArray, ToolEntry);
+
+	return SchemaBuilder.AsJsonObject();
+}
+
+TSharedPtr<FJsonObject> GetFakeToolsetWithMapParamsExpectedSchema()
+{
+	FJsonDomBuilder::FObject SchemaBuilder = MakeSchemaBuilder(
+		FString(TEXT("ToolsetRegistry.FakeToolsetWithMapParams")), FString("1.0"),
+		FString(TEXT("Fake Toolset that has method with TMap parameters.")));
+
+	FJsonDomBuilder::FObject InputSchema = MakeSchema();
+
+	FString ParamName1 = TEXT("mapParamOne");
+	FString ParamName2 = TEXT("mapParamTwo");
+
+	FJsonDomBuilder::FObject InputProperties;
+	InputProperties.Set(ParamName1, MakeMapProperty(TEXT("integer")));
+	InputProperties.Set(ParamName2, MakeMapProperty(TEXT("string")));
+	AddPropertiesToSchema(InputSchema, InputProperties);
+	AddRequiredListToSchema(InputSchema, TArray<FString>({ParamName1, ParamName2}));
+
+	FJsonDomBuilder::FObject OutputSchema = MakeSchema();
+	FJsonDomBuilder::FObject OutputProperties;
+	OutputProperties.Set(ReturnValueName, MakeMapProperty(TEXT("number")));
+	AddPropertiesToSchema(OutputSchema, OutputProperties);
+	AddRequiredListToSchema(OutputSchema, TArray<FString>({ReturnValueName}));
+
+	FJsonDomBuilder::FObject ToolEntry = MakeToolSchemaEntry(
+		TEXT("ToolsetRegistry.FakeToolsetWithMapParams.ToolWithMaps"),
+		TEXT("Method with some params with TMap values."),
+		InputSchema, OutputSchema);
+
+	FJsonDomBuilder::FArray ToolsArray;
+	AddToolSchema(SchemaBuilder, ToolsArray, ToolEntry);
+
+	return SchemaBuilder.AsJsonObject();
+}
+
+// Test that registering with null toolset class input fails properly.
+IMPLEMENT_SIMPLE_AUTOMATION_TEST(
+	FBlueprintLibraryToolsetTestRegisterNullClass,
+	"AI.ToolsetRegistry.BlueprintLibraryToolset.RegisterNullClass",
+	ToolsetRegistryTest::Flags)
+
+bool FBlueprintLibraryToolsetTestRegisterNullClass::RunTest(const FString& UnusedParameters)
+{
+	AddExpectedError(
+		TEXT("Cannot register null toolset class"), EAutomationExpectedErrorFlags::Contains);
+	ToolsetRegistryContainer Container(nullptr);
+	TestTrue("Null toolset class", Container.ToolsetClass == nullptr);
+
+	TSharedPtr<UE::ToolsetRegistry::FBlueprintLibraryToolset> Toolset =
+		Container.ToolsetHandler();
+	TestTrue("Null toolset wrapper", Toolset == nullptr);
+
+	return true;
+}
+
+// Test the blueprint library toolset with null toolset class input fails properly.
+IMPLEMENT_SIMPLE_AUTOMATION_TEST(
+	FBlueprintLibraryToolsetTestConstructWithNullClass,
+	"AI.ToolsetRegistry.BlueprintLibraryToolset.ConstructWithNullClass",
+	ToolsetRegistryTest::Flags)
+
+bool FBlueprintLibraryToolsetTestConstructWithNullClass::RunTest(
+	const FString& UnusedParameters)
+{
+	UE::ToolsetRegistry::FBlueprintLibraryToolset Toolset(nullptr);
+	TestFalse("Invalid blueprint library toolset with null class", Toolset.HasValidTools());
+	return true;
+}
+
+// Test that toolset can unregister and re-register itself properly.
+IMPLEMENT_SIMPLE_AUTOMATION_TEST(
+	FBlueprintLibraryToolsetTestRegisterUnregisterToolset,
+	"AI.ToolsetRegistry.BlueprintLibraryToolset.RegisterUnregisterToolset",
+	ToolsetRegistryTest::Flags)
+
+bool FBlueprintLibraryToolsetTestRegisterUnregisterToolset::RunTest(const FString& UnusedParameters)
+{
+	UClass* ToolsetClass = UFakeToolset::StaticClass();
+
+	TestFalse("Toolset unregistered before registering itself.",
+			  UToolsetRegistry::IsToolsetClassRegistered(ToolsetClass));
+
+	// Scope for registration of ToolsetClass.
+	{
+		ToolsetRegistryContainer Container(ToolsetClass);
+		TestTrue("Toolset registered after registering itself.",
+				 UToolsetRegistry::IsToolsetClassRegistered(ToolsetClass));
+
+		UToolsetRegistry::UnregisterToolsetClass(ToolsetClass);
+		TestFalse("Toolset unregistered after unregistering itself.",
+				  UToolsetRegistry::IsToolsetClassRegistered(ToolsetClass));
+
+		UToolsetRegistry::RegisterToolsetClass(ToolsetClass);
+		TestTrue("Toolset registered after re-registering itself.",
+				 UToolsetRegistry::IsToolsetClassRegistered(ToolsetClass));
+	}
+	return true;
+}
+
+// Test that toolset with no tools fails to register.
+IMPLEMENT_SIMPLE_AUTOMATION_TEST(
+	FBlueprintLibraryToolsetTestNoToolsRegistrationRejected,
+	"AI.ToolsetRegistry.BlueprintLibraryToolset.NoToolsRegistrationRejected",
+	ToolsetRegistryTest::Flags)
+
+bool FBlueprintLibraryToolsetTestNoToolsRegistrationRejected::RunTest(
+	const FString& UnusedParameters)
+{
+	UClass* ToolsetClass = UFakeToolsetWithNoTools::StaticClass();
+
+	ToolsetRegistryContainer Container(ToolsetClass);
+	TestFalse("Toolset with no tools is not registered.",
+			  UToolsetRegistry::IsToolsetClassRegistered(ToolsetClass));
+	return true;
+}
+
+// Test that toolset can return its name, version, and description.
+IMPLEMENT_SIMPLE_AUTOMATION_TEST(
+	FBlueprintLibraryToolsetTestCanGetValues,
+	"AI.ToolsetRegistry.BlueprintLibraryToolset.CanGetValues",
+	ToolsetRegistryTest::Flags)
+
+bool FBlueprintLibraryToolsetTestCanGetValues::RunTest(const FString& UnusedParameters)
+{
+	ToolsetRegistryContainer Container(UFakeToolset::StaticClass());
+	TSharedPtr<UE::ToolsetRegistry::FBlueprintLibraryToolset> Toolset = Container.ToolsetHandler();
+	if (!TestTrue("Valid toolset handler returned", Toolset.IsValid())) return false;
+
+	TestEqual("Toolset returns its own name",
+			  Toolset->GetToolsetName(), FakeToolsetName);
+	TestEqual("Toolset returns default non-overridden version",
+			  Toolset->GetToolsetVersion(), FString("1.0"));
+	TestEqual("Toolset returns description from class definition",
+			  Toolset->GetToolsetDescription(), FString(TEXT("Fake Toolset that does nothing.")));
+	return true;
+}
+
+// Test that toolset can return a basic schema properly.
+IMPLEMENT_SIMPLE_AUTOMATION_TEST(
+	FBlueprintLibraryToolsetTestGetSchema,
+	"AI.ToolsetRegistry.BlueprintLibraryToolset.GetSchema",
+	ToolsetRegistryTest::Flags)
+
+bool FBlueprintLibraryToolsetTestGetSchema::RunTest(const FString& UnusedParameters)
+{
+	ToolsetRegistryContainer Container(UFakeToolset::StaticClass());
+	TSharedPtr<UE::ToolsetRegistry::FBlueprintLibraryToolset> Toolset = Container.ToolsetHandler();
+	if (!TestTrue("Valid toolset handler returned", Toolset.IsValid())) return false;
+
+	TSharedPtr<FJsonObject> ExpectedSchema = GetFakeToolsetExpectedSchema();
+
+	FString ExpectedSchemaJson =
+		UE::ToolsetRegistry::Internal::JsonToString(ExpectedSchema.ToSharedRef());
+	FString OutputSchemaJson = Toolset->GetToolsetJsonSchema();
+	TestEqual("Toolset schema matches expected", OutputSchemaJson, ExpectedSchemaJson);
+	return true;
+}
+
+// Test that toolset with method with default values can return a schema properly.
+IMPLEMENT_SIMPLE_AUTOMATION_TEST(
+	FBlueprintLibraryToolsetTestGetSchemaWithDefaultParams,
+	"AI.ToolsetRegistry.BlueprintLibraryToolset.GetSchemaWithDefaultParams",
+	ToolsetRegistryTest::Flags)
+
+bool FBlueprintLibraryToolsetTestGetSchemaWithDefaultParams::RunTest(
+	const FString& UnusedParameters)
+{
+	ToolsetRegistryContainer Container(UFakeToolsetWithDefaultParams::StaticClass());
+	TSharedPtr<UE::ToolsetRegistry::FBlueprintLibraryToolset> Toolset = Container.ToolsetHandler();
+	if (!TestTrue("Valid toolset handler returned", Toolset.IsValid())) return true;
+
+	TSharedPtr<FJsonObject> ExpectedSchema = GetFakeToolsetWithDefaultParamsExpectedSchema();
+
+	FString ExpectedSchemaJson =
+		UE::ToolsetRegistry::Internal::JsonToString(ExpectedSchema.ToSharedRef());
+	FString OutputSchemaJson = Toolset->GetToolsetJsonSchema();
+	TestEqual("Toolset schema matches expected", OutputSchemaJson, ExpectedSchemaJson);
+	return true;
+}
+
+// Test that toolset with method with array params can return a schema properly.
+IMPLEMENT_SIMPLE_AUTOMATION_TEST(
+	FBlueprintLibraryToolsetTestGetSchemaWithArrayParams,
+	"AI.ToolsetRegistry.BlueprintLibraryToolset.GetSchemaWithArrayParams",
+	ToolsetRegistryTest::Flags)
+
+bool FBlueprintLibraryToolsetTestGetSchemaWithArrayParams::RunTest(const FString& UnusedParameters)
+{
+	ToolsetRegistryContainer Container(UFakeToolsetWithArrayParams::StaticClass());
+	TSharedPtr<UE::ToolsetRegistry::FBlueprintLibraryToolset> Toolset = Container.ToolsetHandler();
+	if (!TestTrue("Valid toolset handler returned", Toolset.IsValid())) return true;
+
+	TSharedPtr<FJsonObject> ExpectedSchema = GetFakeToolsetWithArrayParamsExpectedSchema();
+
+	FString ExpectedSchemaJson =
+		UE::ToolsetRegistry::Internal::JsonToString(ExpectedSchema.ToSharedRef());
+	FString OutputSchemaJson = Toolset->GetToolsetJsonSchema();
+	TestEqual("Toolset schema matches expected", OutputSchemaJson, ExpectedSchemaJson);
+	return true;
+}
+
+// Test that toolset with method with map params can return a schema properly.
+IMPLEMENT_SIMPLE_AUTOMATION_TEST(
+	FBlueprintLibraryToolsetTestGetSchemaWithMapParams,
+	"AI.ToolsetRegistry.BlueprintLibraryToolset.GetSchemaWithMapParams",
+	ToolsetRegistryTest::Flags)
+
+bool FBlueprintLibraryToolsetTestGetSchemaWithMapParams::RunTest(const FString& UnusedParameters)
+{
+	ToolsetRegistryContainer Container(UFakeToolsetWithMapParams::StaticClass());
+	TSharedPtr<UE::ToolsetRegistry::FBlueprintLibraryToolset> Toolset = Container.ToolsetHandler();
+	if (!TestTrue("Valid toolset handler returned", Toolset.IsValid())) return true;
+
+	TSharedPtr<FJsonObject> ExpectedSchema = GetFakeToolsetWithMapParamsExpectedSchema();
+
+	FString ExpectedSchemaJson =
+		UE::ToolsetRegistry::Internal::JsonToString(ExpectedSchema.ToSharedRef());
+	FString OutputSchemaJson = Toolset->GetToolsetJsonSchema();
+	TestEqual("Toolset schema matches expected", OutputSchemaJson, ExpectedSchemaJson);
+	return true;
+}
+
+// Test that toolset with method with set params can return a schema properly.
+IMPLEMENT_SIMPLE_AUTOMATION_TEST(
+	FBlueprintLibraryToolsetTestGetSchemaWithSetParams,
+	"AI.ToolsetRegistry.BlueprintLibraryToolset.GetSchemaWithSetParams",
+	ToolsetRegistryTest::Flags)
+
+bool FBlueprintLibraryToolsetTestGetSchemaWithSetParams::RunTest(const FString& UnusedParameters)
+{
+	ToolsetRegistryContainer Container(UFakeToolsetWithSetParams::StaticClass());
+	TSharedPtr<UE::ToolsetRegistry::FBlueprintLibraryToolset> Toolset = Container.ToolsetHandler();
+	if (!TestTrue("Valid toolset handler returned", Toolset.IsValid())) return true;
+
+	TSharedPtr<FJsonObject> ExpectedSchema = GetFakeToolsetWithSetParamsExpectedSchema();
+
+	FString ExpectedSchemaJson =
+		UE::ToolsetRegistry::Internal::JsonToString(ExpectedSchema.ToSharedRef());
+	FString OutputSchemaJson = Toolset->GetToolsetJsonSchema();
+	TestEqual("Toolset schema matches expected", OutputSchemaJson, ExpectedSchemaJson);
+	return true;
+}
+
+// Test that invoking a tool from a toolset calls the right method and returns the expected result.
+IMPLEMENT_SIMPLE_AUTOMATION_TEST(
+	FBlueprintLibraryToolsetTestExecuteTool,
+	"AI.ToolsetRegistry.BlueprintLibraryToolset.ExecuteTool",
+	ToolsetRegistryTest::Flags)
+
+bool FBlueprintLibraryToolsetTestExecuteTool::RunTest(const FString& UnusedParameters)
+{
+	ToolsetRegistryContainer Container(UFakeToolset::StaticClass());
+	TSharedPtr<UE::ToolsetRegistry::FBlueprintLibraryToolset> Toolset = Container.ToolsetHandler();
+	if (!TestTrue("Valid toolset handler returned", Toolset.IsValid())) return false;
+
+	FString Result;
+	FString Error;
+
+	Toolset->InvokeTool(FakeToolName, FakeToolsetInput, Result, Error);
+	TestTrue("Tool execution returns no error", Error.IsEmpty());
+
+	FString ExpectedOutput = FString(TEXT(R"({"returnValue":"hello"})"));
+	TestEqual("Tool execution returns expected output", Result, ExpectedOutput);
+	return true;
+}
+
+// Test that invoking a tool from an invalid named toolset errors out properly.
+IMPLEMENT_SIMPLE_AUTOMATION_TEST(
+	FBlueprintLibraryToolsetTestExecuteToolInvalidToolsetName,
+	"AI.ToolsetRegistry.BlueprintLibraryToolset.ExecuteToolInvalidToolsetName",
+	ToolsetRegistryTest::Flags)
+
+bool FBlueprintLibraryToolsetTestExecuteToolInvalidToolsetName::RunTest(
+	const FString& UnusedParameters)
+{
+	FString Result;
+	FString Error;
+	AddExpectedError(TEXT("not found"), EAutomationExpectedErrorFlags::Contains);
+	UToolsetRegistry::ExecuteTool(
+		FString("ToolsetRegistry.FakeToolsetInvalid"), FakeToolName, FakeToolsetInput, Result, Error);
+	TestFalse("Tool execution with invalid toolset name returns error", Error.IsEmpty());
+	return true;
+}
+
+// Test that invoking a tool with an invalid tool name errors out properly.
+IMPLEMENT_SIMPLE_AUTOMATION_TEST(
+	FBlueprintLibraryToolsetTestExecuteToolInvalidToolname,
+	"AI.ToolsetRegistry.BlueprintLibraryToolset.ExecuteToolInvalidToolname",
+	ToolsetRegistryTest::Flags)
+
+bool FBlueprintLibraryToolsetTestExecuteToolInvalidToolname::RunTest(
+	const FString& UnusedParameters)
+{
+	ToolsetRegistryContainer Container(UFakeToolset::StaticClass());
+	TSharedPtr<UE::ToolsetRegistry::FBlueprintLibraryToolset> Toolset = Container.ToolsetHandler();
+	TestTrue("Valid toolset handler returned", Toolset.IsValid());
+
+	FString Result;
+	FString Error;
+	UToolsetRegistry::ExecuteTool(
+		FakeToolsetName, FString("SimpleToolInvalid"), FakeToolsetInput, Result, Error);
+	TestFalse("Tool execution with invalid tool name returns error", Error.IsEmpty());
+	return true;
+}
+
+// Test that invoking a tool with a valid but unregistered toolset errors out properly.
+IMPLEMENT_SIMPLE_AUTOMATION_TEST(
+	FBlueprintLibraryToolsetTestExecuteToolUnregisteredToolset,
+	"AI.ToolsetRegistry.BlueprintLibraryToolset.ExecuteToolUnregisteredToolset",
+	ToolsetRegistryTest::Flags)
+
+bool FBlueprintLibraryToolsetTestExecuteToolUnregisteredToolset::RunTest(
+	const FString& UnusedParameters)
+{
+	FString Result;
+	FString Error;
+
+	// Ensure it is unregistered
+	UToolsetRegistry::UnregisterToolsetClass(UFakeToolset::StaticClass());
+
+	AddExpectedError(TEXT("not found"), EAutomationExpectedErrorFlags::Contains);
+	UToolsetRegistry::ExecuteTool(
+		FakeToolsetName, FakeToolName, FakeToolsetInput, Result, Error);
+	TestFalse("Tool execution with unregistered toolset returns error", Error.IsEmpty());
+	return true;
+}
+
+// Test that invoking a tool that doesn't return anything doesn't raise errors.
+IMPLEMENT_SIMPLE_AUTOMATION_TEST(
+	FBlueprintLibraryToolsetTestExecuteToolWithNoReturn,
+	"AI.ToolsetRegistry.BlueprintLibraryToolset.ExecuteToolWithNoReturn",
+	ToolsetRegistryTest::Flags)
+
+bool FBlueprintLibraryToolsetTestExecuteToolWithNoReturn::RunTest(
+	const FString& UnusedParameters)
+{
+	ToolsetRegistryContainer Container(UFakeToolset::StaticClass());
+	TSharedPtr<UE::ToolsetRegistry::FBlueprintLibraryToolset> Toolset = Container.ToolsetHandler();
+	TestTrue("Valid toolset handler returned", Toolset.IsValid());
+
+	FString Result;
+	FString Error;
+	UToolsetRegistry::ExecuteTool(
+		FakeToolsetName, FString("SimpleNoReturnTool"), FakeToolsetInput, Result, Error);
+	TestTrue("Tool execution with no return value runs without error", Error.IsEmpty());
+
+	FString ExpectedOutput = FString(TEXT(R"({"returnValue":null})"));
+	TestEqual("Tool execution returns expected output", Result, ExpectedOutput);
+
+	return true;
+}
+
+// Test that invoking a tool with default values for params provides the proper default value.
+IMPLEMENT_SIMPLE_AUTOMATION_TEST(
+	FBlueprintLibraryToolsetTestExecuteToolWithDefaultParam,
+	"AI.ToolsetRegistry.BlueprintLibraryToolset.ExecuteToolWithDefaultParam",
+	ToolsetRegistryTest::Flags)
+
+bool FBlueprintLibraryToolsetTestExecuteToolWithDefaultParam::RunTest(
+	const FString& UnusedParameters)
+{
+	ToolsetRegistryContainer Container(UFakeToolsetWithDefaultParams::StaticClass());
+	TSharedPtr<UE::ToolsetRegistry::FBlueprintLibraryToolset> Toolset = Container.ToolsetHandler();
+	TestTrue("Valid toolset handler returned", Toolset.IsValid());
+
+	FString Result;
+	FString Error;
+	UToolsetRegistry::ExecuteTool(
+		FString("ToolsetRegistry.FakeToolsetWithDefaultParams"), FString("ToolWithDefaults"),
+		FakeToolsetInput, Result, Error);
+	TestTrue("Tool execution with default param returns without error", Error.IsEmpty());
+
+	FString ExpectedOutput = FString(TEXT(R"({"returnValue":"hello: 42"})"));
+	TestEqual("Tool execution returns expected output", Result, ExpectedOutput);
+
+	return true;
+}
+
+// Test that invoking a tool with params that don't match input schema errors out properly.
+IMPLEMENT_SIMPLE_AUTOMATION_TEST(
+	FBlueprintLibraryToolsetTestExecuteToolWithBadInput,
+	"AI.ToolsetRegistry.BlueprintLibraryToolset.ExecuteToolWithBadInput",
+	ToolsetRegistryTest::Flags)
+
+bool FBlueprintLibraryToolsetTestExecuteToolWithBadInput::RunTest(
+	const FString& UnusedParameters)
+{
+	ToolsetRegistryContainer Container(UFakeToolset::StaticClass());
+	TSharedPtr<UE::ToolsetRegistry::FBlueprintLibraryToolset> Toolset = Container.ToolsetHandler();
+	TestTrue("Valid toolset handler returned", Toolset.IsValid());
+
+	FString Input = TEXT(R"({"not_expected":"bad_value"})");
+	FString Result;
+	FString Error;
+	UToolsetRegistry::ExecuteTool(
+		FakeToolsetName, FakeToolName, Input, Result, Error);
+	TestFalse("Tool execution with invalid param returns with error", Error.IsEmpty());
+	return true;
+}
+
+// Test that a toolset doesn't include private, internal and non-static methods.
+IMPLEMENT_SIMPLE_AUTOMATION_TEST(
+	FBlueprintLibraryToolsetTestValidateToolMethods,
+	"AI.ToolsetRegistry.BlueprintLibraryToolset.ValidateToolMethods",
+	ToolsetRegistryTest::Flags)
+
+bool FBlueprintLibraryToolsetTestValidateToolMethods::RunTest(
+	const FString& UnusedParameters)
+{
+	ToolsetRegistryContainer Container(UFakeToolsetWithPrivateMethods::StaticClass());
+	TSharedPtr<UE::ToolsetRegistry::FBlueprintLibraryToolset> Toolset = Container.ToolsetHandler();
+	if (!TestTrue("Valid toolset handler returned", Toolset.IsValid()))
+	{
+		return false;
+	}
+
+	TArray<FString> ToolNames = Toolset->GetToolNames();
+
+	if (!TestEqual("Only one valid tool should be included", ToolNames.Num(), 1))
+	{
+		return false;
+	}
+	TestEqual("Only \"SimpleTool\" should be valid tool", ToolNames[0], FString("SimpleTool"));
+
+	return true;
+}
+
+// Invoke a C++ tool that produces errors, not from bad input.
+IMPLEMENT_SIMPLE_AUTOMATION_TEST(
+	FBlueprintLibraryToolsetTestExecuteCppToolWithError,
+	"AI.ToolsetRegistry.BlueprintLibraryToolset.ExecuteCppToolWithError",
+	ToolsetRegistryTest::Flags)
+
+bool FBlueprintLibraryToolsetTestExecuteCppToolWithError::RunTest(
+	const FString& UnusedParameters)
+{
+	ToolsetRegistryContainer Container(UFakeErrorProneToolset::StaticClass());
+	TSharedPtr<UE::ToolsetRegistry::FBlueprintLibraryToolset> Toolset = Container.ToolsetHandler();
+	if (!TestTrue("Valid toolset handler returned", Toolset.IsValid()))
+	{
+		return false;
+	}
+
+	FString Result;
+	FString Error;
+	UToolsetRegistry::ExecuteTool(Toolset->GetToolsetName(), TEXT("BuggyTool"), TEXT(""), Result, Error);
+	TestFalse("Tool execution with buggy tool returns with error", Error.IsEmpty());
+	TestEqual("Tool execution with buggy tool returns correct raised exception",
+			  Error, UFakeErrorProneToolset::ErrorMessage);
+
+	return true;
+}
+
+// Invoke python tool that produces errors.
+IMPLEMENT_SIMPLE_AUTOMATION_TEST(
+	FBlueprintLibraryToolsetTestExecutePythonToolWithError,
+	"AI.ToolsetRegistry.BlueprintLibraryToolset.ExecutePythonToolWithError",
+	ToolsetRegistryTest::Flags)
+
+bool FBlueprintLibraryToolsetTestExecutePythonToolWithError::RunTest(
+	const FString& UnusedParameters)
+{
+	if (!IPythonScriptPlugin::Get()->IsPythonAvailable())
+	{
+		return true;
+	}
+
+	FPythonCommandEx PythonCommand = FPythonCommandEx();
+	PythonCommand.ExecutionMode = EPythonCommandExecutionMode::ExecuteStatement;
+	PythonCommand.Command = TEXT(
+		"from toolset_registry.tests.demo_toolset import ErrorProneToolset; "
+		"import unreal; unreal.ToolsetRegistry.register_toolset_class(ErrorProneToolset)");
+	if (!TestTrue("Register ErrorProneToolset", IPythonScriptPlugin::Get()->ExecPythonCommandEx(PythonCommand)))
+	{
+		return false;
+	}
+
+	FString Input = TEXT(R"({"msg":"Hi"})");
+	FString Result;
+	FString Error;
+	FString ExpectedError = TEXT("Bad tool bad");
+	UToolsetRegistry::ExecuteTool(
+		"toolset_registry.tests.demo_toolset.ErrorProneToolset", "bad_tool", Input, Result, Error);
+
+	TestFalse("Tool error was returned", Error.IsEmpty());
+	TestEqual("Error is expected exception message", Error, ExpectedError);
+
+	PythonCommand.Command = TEXT(
+		"from toolset_registry.tests.demo_toolset import ErrorProneToolset; "
+		"import unreal; unreal.ToolsetRegistry.unregister_toolset_class(ErrorProneToolset)");
+	TestTrue("Unregister ErrorProneToolset", IPythonScriptPlugin::Get()->ExecPythonCommandEx(PythonCommand));
+
+	return true;
+}
+
+IMPLEMENT_SIMPLE_AUTOMATION_TEST(
+	FBlueprintLibraryToolsetTestToolsetNaming,
+	"AI.ToolsetRegistry.BlueprintLibraryToolset.ToolsetNaming",
+	ToolsetRegistryTest::Flags)
+
+bool FBlueprintLibraryToolsetTestToolsetNaming::RunTest(
+	const FString& UnusedParameters)
+{
+	UClass* ToolsetClass = UFakeToolset::StaticClass();
+	FString ToolsetName = UE::ToolsetRegistry::FBlueprintLibraryToolset::GetToolsetClassName(ToolsetClass);
+	FString PackageName = ToolsetClass->GetOuter()->GetName();
+	FString PathName = ToolsetClass->GetPathName();
+	FString ClassName = ToolsetClass->GetName();
+	TestNotEqual("Toolset name should not be Class name", ToolsetName, ClassName);
+	TestNotEqual("Toolset name should not be same as Path name", ToolsetName, PathName);
+	TestTrue("Toolset name should include Class name", ToolsetName.Contains(ClassName));
+	int32 DotIndex;
+	TestTrue("Toolset name should include a namespace qualifier", ToolsetName.FindLastChar('.', DotIndex));
+	FString ToolsetNamespace = ToolsetName.Left(DotIndex);
+	FString ToolsetNameClass = ToolsetName.RightChop(DotIndex + 1);
+	TestEqual("Toolset name should have class name last", ToolsetNameClass, ClassName);
+	int32 SlashIndex;
+	TestTrue("Toolset package should be slash-separated", PackageName.FindLastChar('/', SlashIndex));
+	FString PackageLast = PackageName.RightChop(SlashIndex + 1);
+	TestEqual("Toolset name namespace should be rightmost package part", ToolsetNamespace, PackageLast);
+	return true;
+}
+
+IMPLEMENT_SIMPLE_AUTOMATION_TEST(
+	FBlueprintLibraryToolsetTestToolsetNamingNull,
+	"AI.ToolsetRegistry.BlueprintLibraryToolset.ToolsetNamingNull",
+	ToolsetRegistryTest::Flags)
+
+bool FBlueprintLibraryToolsetTestToolsetNamingNull::RunTest(
+	const FString& UnusedParameters)
+{
+	FString ToolsetName = UE::ToolsetRegistry::FBlueprintLibraryToolset::GetToolsetClassName(nullptr);
+	return TestTrue("Toolset name for null UClass should be empty", ToolsetName.IsEmpty());
+}
+
+// Run python tests
+IMPLEMENT_SIMPLE_AUTOMATION_TEST(
+	FBlueprintLibraryToolsetTestPythonTests,
+	"AI.ToolsetRegistry.BlueprintLibraryToolset.PythonTests",
+	ToolsetRegistryTest::Flags)
+
+bool FBlueprintLibraryToolsetTestPythonTests::RunTest(
+	const FString& UnusedParameters)
+{
+	if (!IPythonScriptPlugin::Get()->IsPythonAvailable())
+	{
+		return true;
+	}
+
+	FPythonCommandEx PythonCommand = FPythonCommandEx();
+	PythonCommand.ExecutionMode = EPythonCommandExecutionMode::ExecuteStatement;
+	PythonCommand.Command = TEXT("import toolset_registry.tests; toolset_registry.tests.run_tests()");
+	return IPythonScriptPlugin::Get()->ExecPythonCommandEx(PythonCommand);
+}
+
+#endif	// WITH_DEV_AUTOMATION_TESTS
+
diff --git a/Engine/Plugins/Experimental/ToolsetRegistry/Source/ToolsetRegistry/Private/Tests/BlueprintLibraryToolsetTest.h b/Engine/Plugins/Experimental/ToolsetRegistry/Source/ToolsetRegistry/Private/Tests/BlueprintLibraryToolsetTest.h
new file mode 100644
--- /dev/null
+++ b/Engine/Plugins/Experimental/ToolsetRegistry/Source/ToolsetRegistry/Private/Tests/BlueprintLibraryToolsetTest.h
@@ -0,0 +1,171 @@
+// Copyright Epic Games, Inc. All Rights Reserved.
+
+#pragma once
+
+#include "CoreMinimal.h"
+
+#include "Dom/JsonObject.h"
+#include "Kismet/KismetSystemLibrary.h"
+
+#include "BPToolsetRegistry.h"
+
+#include "BlueprintLibraryToolsetTest.generated.h"
+
+/** Fake Toolset that does nothing. */
+UCLASS(Blueprintable)
+class UFakeToolset : public UToolsetDefinition
+{
+	GENERATED_BODY()
+
+public:
+
+	/** Simple tool that does nothing. */
+	UFUNCTION(BlueprintCallable, Category=Fake)
+	static FString SimpleTool(const FString& InString)
+	{
+		return InString;
+	}
+
+	/** Simple tool that does nothing and returns nothing. */
+	UFUNCTION(BlueprintCallable, Category=Fake)
+	static void SimpleNoReturnTool(const FString& InString)
+	{
+		return;
+	}
+
+	/** UFunction for internal use, should not be marked as tool. */
+	UFUNCTION(meta=(BlueprintInternalUseOnly=true))
+	static void SimpleInternalMethod() { }
+};
+
+static TSharedPtr<FJsonObject> GetFakeToolsetExpectedSchema();
+
+/** Fake Toolset that has method with default parameters. */
+UCLASS(Blueprintable)
+class UFakeToolsetWithDefaultParams : public UToolsetDefinition
+{
+	GENERATED_BODY()
+
+public:
+
+	/** Method with some params with default values. */
+	UFUNCTION(BlueprintCallable, Category=Fake)
+	static FString ToolWithDefaults(const FString& InString, int InValue=42)
+	{
+		return FString::Printf(TEXT("%s: %d"), *InString, InValue);
+	}
+};
+
+static TSharedPtr<FJsonObject> GetFakeToolsetWithDefaultParamsExpectedSchema();
+
+
+/** Fake Toolset with no tools. */
+UCLASS(Blueprintable)
+class UFakeToolsetWithNoTools : public UToolsetDefinition
+{
+	GENERATED_BODY()
+};
+
+static TSharedPtr<FJsonObject> GetFakeToolsetWithNoToolsExpectedSchema();
+
+/** Fake Toolset that has method with TArray parameters. */
+UCLASS(Blueprintable)
+class UFakeToolsetWithArrayParams : public UToolsetDefinition
+{
+	GENERATED_BODY()
+
+public:
+
+	/** Method with some params with TArray values. */
+	UFUNCTION(BlueprintCallable, Category=Fake)
+	static TArray<float> ToolWithArrays(const TArray<int> ArrayParamOne, TArray<FString> ArrayParamTwo)
+	{
+		return TArray<float>();
+	}
+};
+
+static TSharedPtr<FJsonObject> GetFakeToolsetWithArrayParamsExpectedSchema();
+
+/** Fake Toolset that has method with TMap parameters. */
+UCLASS(Blueprintable)
+class UFakeToolsetWithMapParams : public UToolsetDefinition
+{
+	GENERATED_BODY()
+
+public:
+
+	/** Method with some params with TMap values. */
+	UFUNCTION(BlueprintCallable, Category=Fake)
+	static TMap<FString, float> ToolWithMaps(const TMap<FString, int> MapParamOne, TMap<FString, FString> MapParamTwo)
+	{
+		return TMap<FString, float>();
+	}
+};
+
+static TSharedPtr<FJsonObject> GetFakeToolsetWithMapParamsExpectedSchema();
+
+/** Fake Toolset that has method with TSet parameters. */
+UCLASS(Blueprintable)
+class UFakeToolsetWithSetParams : public UToolsetDefinition
+{
+	GENERATED_BODY()
+
+public:
+
+	/** Method with some params with TSet values. */
+	UFUNCTION(BlueprintCallable, Category=Fake)
+	static TSet<float> ToolWithSets(const TSet<int> SetParamOne, TSet<FString> SetParamTwo)
+	{
+		return TSet<float>();
+	}
+};
+
+static TSharedPtr<FJsonObject> GetFakeToolsetWithSetParamsExpectedSchema();
+
+/** Fake Toolset with private and non-static methods. */
+UCLASS(Blueprintable)
+class UFakeToolsetWithPrivateMethods : public UToolsetDefinition
+{
+	GENERATED_BODY()
+
+public:
+
+	/** Simple tool that does nothing. */
+	UFUNCTION(BlueprintCallable, Category=Fake)
+	static FString SimpleTool(const FString& InString)
+	{
+		return InString;
+	}
+
+	/** UFunction for internal use, should not be marked as tool. */
+	UFUNCTION(meta=(BlueprintInternalUseOnly=true))
+	static void SimpleInternalMethod() {}
+
+	/** UFunction non-static method, should not be marked as tool. */
+	UFUNCTION()
+	void SimpleNonStaticMethod() {}
+
+
+private:
+	/** UFunction private method, should not be marked as tool. */
+	UFUNCTION()
+	static void SimplePrivateMethod() {}
+};
+
+/** Fake Toolset that has error-prone method. */
+UCLASS(Blueprintable)
+class UFakeErrorProneToolset : public UToolsetDefinition
+{
+	GENERATED_BODY()
+
+public:
+
+	static inline FString ErrorMessage = TEXT("Bug in this tool");
+
+	/** Method that raises errors. */
+	UFUNCTION(BlueprintCallable, Category=Fake)
+	static void BuggyTool()
+	{
+		UKismetSystemLibrary::RaiseScriptError(ErrorMessage);
+	}
+};
diff --git a/Engine/Plugins/Experimental/ToolsetRegistry/Source/ToolsetRegistry/Private/Tests/DelegateHandleTest.cpp b/Engine/Plugins/Experimental/ToolsetRegistry/Source/ToolsetRegistry/Private/Tests/DelegateHandleTest.cpp
new file mode 100644
--- /dev/null
+++ b/Engine/Plugins/Experimental/ToolsetRegistry/Source/ToolsetRegistry/Private/Tests/DelegateHandleTest.cpp
@@ -0,0 +1,142 @@
+// Copyright Epic Games, Inc. All Rights Reserved.
+
+#include "Containers/UnrealString.h"
+#include "Delegates/DelegateCombinations.h"
+#include "Misc/AutomationTest.h"
+#include "Templates/Function.h"
+#include "Templates/UniquePtr.h"
+#include "Templates/UnrealTemplate.h"
+
+#include "ToolsetRegistryDelegateHandle.h"
+#include "ToolsetRegistryTestFlags.h"
+
+#if WITH_DEV_AUTOMATION_TESTS
+namespace UE::ToolsetRegistry
+{
+	namespace
+	{
+		struct FCallCounter
+		{
+			int CallCount = 0;
+
+			void Call()
+			{
+				CallCount++;
+			}
+
+			TFunction<void()> Callable()
+			{
+				return [this]() -> void { Call(); };
+			}
+		};
+	}
+
+	DECLARE_MULTICAST_DELEGATE(FToolsetRegistryOnTest)
+
+	BEGIN_DEFINE_SPEC(FToolsetRegistryDelegateHandleRaii, "AI.ToolsetRegistry.DelegateHandleRaii",
+		ToolsetRegistryTest::Flags)
+		END_DEFINE_SPEC(FToolsetRegistryDelegateHandleRaii)
+
+		void FToolsetRegistryDelegateHandleRaii::Define()
+	{
+		Describe(TEXT("TDelegateHandleRaii"), [this]
+			{
+				It(TEXT("can construct a container attached to a delegate handle"), [this]
+					{
+						FToolsetRegistryOnTest OnTest;
+						FCallCounter CallCounter;
+						auto RawDelegateHandle = OnTest.AddLambda(CallCounter.Callable());
+						auto RawDelegateHandleToMove = RawDelegateHandle;
+						TDelegateHandleRaii Handle(OnTest, MoveTemp(RawDelegateHandleToMove));
+						(void)TestEqual(TEXT("Delegate handle"), Handle.Get(), RawDelegateHandle);
+					});
+
+				It(TEXT("can construct a container that references an invalid handle"), [this]
+					{
+						FToolsetRegistryOnTest OnTest;
+						FDelegateHandle RawDelegateHandle;
+						TDelegateHandleRaii Handle(OnTest, MoveTemp(RawDelegateHandle));
+						(void)TestFalse(
+							TEXT("Delegate handle is not valid"), Handle.Get().IsValid());
+					});
+
+				It(TEXT("should unregister a delegate handler on reset"), [this]
+					{
+						FToolsetRegistryOnTest OnTest;
+						FCallCounter CallCounter;
+						TDelegateHandleRaii Handle(
+							OnTest, OnTest.AddLambda(CallCounter.Callable()));
+						OnTest.Broadcast();
+						(void)TestEqual(TEXT("Called once"), CallCounter.CallCount, 1);
+						Handle.Reset();
+
+						OnTest.Broadcast();
+						(void)TestEqual(TEXT("Still called once"), CallCounter.CallCount, 1);
+						(void)TestFalse(
+							TEXT("Delegate handle is not valid"),
+							Handle.Get().IsValid());
+					});
+
+				It(TEXT("should unregister a delegate handler on destruction"), [this]
+					{
+						FToolsetRegistryOnTest OnTest;
+						FCallCounter CallCounter;
+						{
+							TDelegateHandleRaii Handle(
+								OnTest, OnTest.AddLambda(CallCounter.Callable()));
+							OnTest.Broadcast();
+							(void)TestEqual(TEXT("Called once"), CallCounter.CallCount, 1);
+						}
+
+						OnTest.Broadcast();
+						(void)TestEqual(TEXT("Still called once"), CallCounter.CallCount, 1);
+					});
+			});
+
+		Describe(TEXT("FDelegateHandleRaii::Create"), [this]
+			{
+				It(TEXT("should create a container attached to a delegate handle"), [this]
+					{
+						FToolsetRegistryOnTest OnTest;
+						FCallCounter CallCounter;
+						auto RawDelegateHandle = OnTest.AddLambda(CallCounter.Callable());
+						auto RawDelegateHandleToMove = RawDelegateHandle;
+						auto Handle = FDelegateHandleRaii::Create(
+							OnTest, MoveTemp(RawDelegateHandleToMove));
+						if (!TestTrue(TEXT("Handle created"), Handle.IsValid())) return;
+						(void)TestEqual(TEXT("Delegate handle"),
+							Handle->Get(), RawDelegateHandle);
+					});
+			});
+
+		Describe(TEXT("FDelegateHandleRaii::IsValid"), [this]
+			{
+				It(TEXT("should return false if it does not reference a valid pointer"), [this]
+					{
+						FDelegateHandleRaii Handle;
+						(void)TestFalse(TEXT("Not valid on construction"), Handle.IsValid());
+					});
+
+				It(TEXT("should return true if it references a valid pointer and handle"), [this]
+					{
+						FToolsetRegistryOnTest OnTest;
+						FCallCounter CallCounter;
+						auto Handle = FDelegateHandleRaii::Create(
+							OnTest, OnTest.AddLambda(CallCounter.Callable()));
+						(void)TestTrue(TEXT("Valid after create"), Handle.IsValid());
+					});
+
+				It(TEXT("should return false if it references an invalid handle"), [this]
+					{
+						FToolsetRegistryOnTest OnTest;
+						FCallCounter CallCounter;
+						auto Handle = FDelegateHandleRaii::Create(
+							OnTest, OnTest.AddLambda(CallCounter.Callable()));
+						Handle->Reset();
+						(void)TestFalse(TEXT("Not valid after reset"), Handle.IsValid());
+					});
+			});
+	}
+
+}
+#endif  // WITH_DEV_AUTOMATION_TESTS
diff --git a/Engine/Plugins/Experimental/ToolsetRegistry/Source/ToolsetRegistry/Private/Tests/ObjectFunctionToolCallTest.cpp b/Engine/Plugins/Experimental/ToolsetRegistry/Source/ToolsetRegistry/Private/Tests/ObjectFunctionToolCallTest.cpp
new file mode 100644
--- /dev/null
+++ b/Engine/Plugins/Experimental/ToolsetRegistry/Source/ToolsetRegistry/Private/Tests/ObjectFunctionToolCallTest.cpp
@@ -0,0 +1,227 @@
+// Copyright Epic Games, Inc. All Rights Reserved.
+
+#include "ObjectFunctionToolCallTest.h"
+
+#include "Misc/AutomationTest.h"
+#include "UObject/Package.h"
+#include "Serialization/JsonSerializer.h"
+#include "Serialization/JsonWriter.h"
+
+#include "Tests/ToolsetRegistryTestFlags.h"
+
+#include "ToolsetRegistryModule.h"
+#include "ToolsetRegistry/JsonConversion.h"
+#include "ToolsetRegistry/ObjectFunctionToolCall.h"
+
+PRAGMA_DISABLE_DEPRECATION_WARNINGS
+
+#define DEBUG_PRINT_JSONS_TO_LOG 0
+
+#if WITH_DEV_AUTOMATION_TESTS
+
+namespace UE::ToolsetRegistry
+{
+	namespace
+	{
+		void PrintJsonToLog(
+			const TCHAR* Name, const TSharedPtr<FJsonObject>& JsonObject)
+		{
+			if (JsonObject)
+			{
+				FString JsonString =
+					UE::ToolsetRegistry::Internal::JsonToString(JsonObject.ToSharedRef());
+				UE_LOG(LogToolsetRegistry, Display, TEXT("Json \"%s\" is valid -\n%s"), Name, *JsonString);
+			}
+			else
+			{
+				UE_LOG(LogToolsetRegistry, Error, TEXT("Json \"%s\" is invalid."), Name);
+			}
+		}
+	}
+
+	// Test function with required param.
+
+	IMPLEMENT_SIMPLE_AUTOMATION_TEST(
+		FObjectFunctionToolCallTestExecuteToolCallFunctionWithRequiredParam,
+		"AI.ToolsetRegistry.ToolCall.ExecuteToolCallFunctionWithRequiredParam",
+		ToolsetRegistryTest::Flags);
+
+	bool FObjectFunctionToolCallTestExecuteToolCallFunctionWithRequiredParam::RunTest(
+		const FString& Parameters)
+	{
+		static const FName FunctionName = "TestFuncWithRequiredParam";
+		static constexpr float ClassFloatValue = 20.0f;
+		static constexpr float FunctionInputFloatValue = 3.0f;
+		static const FString FunctionInputParamsJsonString =
+			FString::Printf(TEXT("{\n\"InFloat\": %f\n}"), FunctionInputFloatValue);
+
+		// Create an object and provide it specific internal state. 
+		// (For this test - Using TStrongObjectPtr to prevent GC.)
+
+		const TStrongObjectPtr InstanceObject =
+			TStrongObjectPtr(NewObject<UToolCallFakeClass>());
+
+		InstanceObject->Float = ClassFloatValue;
+
+		// Execute.
+
+		const FObjectFunctionToolCall ToolCall =
+			FObjectFunctionToolCall::ConstructWithNamedClassFunction(InstanceObject.Get(), FunctionName);
+
+		FObjectFunctionToolCall::FFunctionInputParamsJson FunctionInputParamsJson;
+		FunctionInputParamsJson.Set<FString>(FunctionInputParamsJsonString);
+		const FObjectFunctionToolCall::FToolCallResult ToolCallResult =
+			ToolCall.Execute(FunctionInputParamsJson);
+
+		if (!TestTrue("Execute tool call result does not have error.",
+			!ToolCallResult.HasError()))
+		{
+			UE_LOG(LogToolsetRegistry, Error, TEXT("Tool call error - %s"),
+				*ToolCallResult.GetError());
+			return false;
+		}
+		if (!TestTrue("Execute tool call result has value.",
+			ToolCallResult.HasValue()))
+		{
+			return false;
+		}
+
+		const TSharedPtr<FJsonObject> JsonObject = ToolCallResult.GetValue();
+
+		if (!TestTrue("Execute tool call result has valid value.",
+			JsonObject.IsValid()))
+		{
+			return false;
+		}
+
+#if DEBUG_PRINT_JSONS_TO_LOG
+		PrintJsonToLog(TEXT("ToolCallWithRequiredParamReturnedJsonObject"),
+			JsonObject.ToSharedRef());
+#endif
+
+		if (!TestTrue("Return object is not empty (has values.)",
+			JsonObject->Values.Num() > 0))
+		{
+			return false;
+		}
+		if (!TestTrue(TEXT("Return object has float returnValue."),
+			JsonObject->HasTypedField(TEXT("returnValue"), EJson::Number)))
+		{
+			return false;
+		}
+
+		// Check.
+
+		const float ReturnValueFloat =
+			JsonObject->GetNumberField(TEXT("returnValue"));
+		const float ReturnValueFloatExpected =
+			InstanceObject->TestFuncWithRequiredParam(FunctionInputFloatValue);
+		const FString TestMessage = FString::Printf(TEXT("Return value %f is expected value %f."),
+			ReturnValueFloat, ReturnValueFloatExpected);
+		if (!TestTrue(TestMessage,
+			ReturnValueFloat == ReturnValueFloatExpected))
+		{
+			return false;
+		}
+
+		// Success.
+
+		return true;
+	}
+
+	// Test function with optional param.
+
+	IMPLEMENT_SIMPLE_AUTOMATION_TEST(
+		FObjectFunctionToolCallTestExecuteToolCallFunctionWithOptionalParam,
+		"AI.ToolsetRegistry.ToolCall.ExecuteToolCallFunctionOptionalParam",
+		ToolsetRegistryTest::Flags);
+
+	bool FObjectFunctionToolCallTestExecuteToolCallFunctionWithOptionalParam::RunTest(
+		const FString& Parameters)
+	{
+		static const FName FunctionName = "TestFuncWithOptionalParam";
+		static constexpr float ClassFloatValue = 20.0f;
+
+		// Create an object and provide it specific internal state. 
+		// (For this test - Using TStrongObjectPtr to prevent GC.)
+
+		const TStrongObjectPtr InstanceObject =
+			TStrongObjectPtr(NewObject<UToolCallFakeClass>());
+
+		InstanceObject->Float = ClassFloatValue;
+
+		// Execute.
+
+		const TOptional<FObjectFunctionToolCall> ToolCall =
+			FObjectFunctionToolCall::ConstructWithNamedClassFunction(InstanceObject.Get(), FunctionName);
+		if (!TestTrue("Tool call created from class and function name.",
+			ToolCall.IsSet()))
+		{
+			return false;
+		}
+
+		const FObjectFunctionToolCall::FToolCallResult ToolCallResult = ToolCall->Execute();
+
+		if (!TestTrue("Execute tool call result does not have error.",
+			!ToolCallResult.HasError()))
+		{
+			UE_LOG(LogToolsetRegistry, Error, TEXT("Tool call error - %s"),
+				*ToolCallResult.GetError());
+			return false;
+		}
+		if (!TestTrue("Execute tool call result has value.",
+			ToolCallResult.HasValue()))
+		{
+			return false;
+		}
+
+		const TSharedPtr<FJsonObject> JsonObject = ToolCallResult.GetValue();
+
+		if (!TestTrue("Execute tool call result has valid value.",
+			JsonObject.IsValid()))
+		{
+			return false;
+		}
+
+#if DEBUG_PRINT_JSONS_TO_LOG
+		PrintJsonToLog(TEXT("ToolCallWithOptionalParamReturnedJsonObject"),
+			JsonObject.ToSharedRef());
+#endif
+
+		if (!TestTrue("Return object is not empty (has values.)",
+			JsonObject->Values.Num() > 0))
+		{
+			return false;
+		}
+		if (!TestTrue(TEXT("Return object has float returnValue."),
+			JsonObject->HasTypedField(TEXT("returnValue"), EJson::Number)))
+		{
+			return false;
+		}
+
+		// Check.
+
+		const float ReturnValueFloat =
+			JsonObject->GetNumberField(TEXT("returnValue"));
+		const float ReturnValueFloatExpected =
+			InstanceObject->TestFuncWithOptionalParam();
+		const FString TestMessage =
+			FString::Printf(TEXT("Return value %f is expected value %f."), ReturnValueFloat, ReturnValueFloatExpected);
+		if (!TestTrue(TestMessage,
+			ReturnValueFloat == ReturnValueFloatExpected))
+		{
+			return false;
+		}
+
+		// Success.
+
+		return true;
+	}
+
+}
+
+#endif
+
+PRAGMA_ENABLE_DEPRECATION_WARNINGS
+
+
diff --git a/Engine/Plugins/Experimental/ToolsetRegistry/Source/ToolsetRegistry/Private/Tests/ObjectFunctionToolCallTest.h b/Engine/Plugins/Experimental/ToolsetRegistry/Source/ToolsetRegistry/Private/Tests/ObjectFunctionToolCallTest.h
new file mode 100644
--- /dev/null
+++ b/Engine/Plugins/Experimental/ToolsetRegistry/Source/ToolsetRegistry/Private/Tests/ObjectFunctionToolCallTest.h
@@ -0,0 +1,30 @@
+// Copyright Epic Games, Inc. All Rights Reserved.
+
+#pragma once
+
+#include "CoreMinimal.h"
+
+#include "ObjectFunctionToolCallTest.generated.h"
+
+UCLASS(BlueprintType)
+class UToolCallFakeClass : public UObject
+{
+	GENERATED_BODY()
+	
+public:
+	
+	UPROPERTY(EditAnywhere, Category="Main")
+	float Float = 0.0f;
+	
+	UFUNCTION(BlueprintCallable, Category="Main") 
+	float TestFuncWithRequiredParam(const float InFloat) const 
+	{
+		return 100.0f + Float + InFloat; 
+	}
+	
+	UFUNCTION(BlueprintCallable, Category="Main") 
+	float TestFuncWithOptionalParam(const float InFloat = 3.0f) const 
+	{
+		return 100.0f + Float + InFloat; 
+	}
+};
diff --git a/Engine/Plugins/Experimental/ToolsetRegistry/Source/ToolsetRegistry/Private/Tests/ToolsetJsonTest.cpp b/Engine/Plugins/Experimental/ToolsetRegistry/Source/ToolsetRegistry/Private/Tests/ToolsetJsonTest.cpp
--- a/Engine/Plugins/Experimental/ToolsetRegistry/Source/ToolsetRegistry/Private/Tests/ToolsetJsonTest.cpp
+++ b/Engine/Plugins/Experimental/ToolsetRegistry/Source/ToolsetRegistry/Private/Tests/ToolsetJsonTest.cpp
@@ -1,78 +1,84 @@
 // Copyright Epic Games, Inc. All Rights Reserved.
 
 #include "ToolsetJsonTest.h"
-#include "CoreMinimal.h"
 
+#include "CoreMinimal.h"
 #include "JsonObjectConverter.h"
 #include "Misc/AutomationTest.h"
-#include "ToolsetJson.h"
 
+#include "ToolsetRegistry/ToolsetJson.h"
 
 #if WITH_DEV_AUTOMATION_TESTS
 
-void TestObjectRefSchema(FAutomationTestBase& Test, const TSharedPtr<FJsonObject>& Schema, UClass* Class, bool IsClass)
+namespace
 {
-	Test.TestTrue("Schema", Schema.IsValid());
-	if (!Schema.IsValid())
+	void TestObjectRefSchema(
+		FAutomationTestBase& Test, const TSharedPtr<FJsonObject>& Schema,
+		UClass* Class, bool IsClass)
 	{
-		return;
-	}
+		Test.TestTrue("Schema", Schema.IsValid());
+		if (!Schema.IsValid())
+		{
+			return;
+		}
 
-	FString Type;
-	Test.TestTrue("Type", Schema->TryGetStringField(TEXT("type"), Type));
-	Test.TestEqual("Type", Type, FString(TEXT("object")));
+		FString Type;
+		Test.TestTrue("Type", Schema->TryGetStringField(TEXT("type"), Type));
+		Test.TestEqual("Type", Type, FString(TEXT("object")));
 
-	if (Class)
-	{
-		FString Title;
-		Test.TestTrue("Title", Schema->TryGetStringField(TEXT("title"), Title));
-		FString ClassPrefix(IsClass ? TEXT("Class@") : TEXT(""));
-		Test.TestEqual("Title", Title, ClassPrefix + Class->GetPathName());
-	}	
+		if (Class)
+		{
+			FString Title;
+			Test.TestTrue("Title", Schema->TryGetStringField(TEXT("title"), Title));
+			FString ClassPrefix(IsClass ? TEXT("Class@") : TEXT(""));
+			Test.TestEqual("Title", Title, ClassPrefix + Class->GetPathName());
+		}	
 	
-	const TSharedPtr<FJsonObject>* Properties = nullptr;
-	Test.TestTrue("Properties", Schema->TryGetObjectField(TEXT("properties"), Properties));
-	if (Properties)
-	{
-		const TSharedPtr<FJsonObject>* RefPath = nullptr;
-		Test.TestTrue("RefPath", Properties->Get()->TryGetObjectField(TEXT("refPath"), RefPath));
-		if (RefPath)
+		const TSharedPtr<FJsonObject>* Properties = nullptr;
+		Test.TestTrue("Properties", Schema->TryGetObjectField(TEXT("properties"), Properties));
+		if (Properties)
 		{
-			Test.TestTrue("RefPath", RefPath->Get()->TryGetStringField(TEXT("type"), Type));
-			Test.TestEqual("RefPath", Type, FString(TEXT("string")));
+			const TSharedPtr<FJsonObject>* RefPath = nullptr;
+			Test.TestTrue(
+				"RefPath", Properties->Get()->TryGetObjectField(TEXT("refPath"), RefPath));
+			if (RefPath)
+			{
+				Test.TestTrue("RefPath", RefPath->Get()->TryGetStringField(TEXT("type"), Type));
+				Test.TestEqual("RefPath", Type, FString(TEXT("string")));
+			}
 		}
 	}
-}
 
-void TestObjectRefData(
-	FAutomationTestBase& Test, const TSharedPtr<FJsonObject>& Schema, const FString& Expected)
-{
-	Test.TestTrue("Schema", Schema.IsValid());
-	if (!Schema.IsValid())
+	void TestObjectRefData(
+		FAutomationTestBase& Test, const TSharedPtr<FJsonObject>& Schema, const FString& Expected)
 	{
-		return;
-	}
+		Test.TestTrue("Schema", Schema.IsValid());
+		if (!Schema.IsValid())
+		{
+			return;
+		}
 	
-	FString RefPath;
-	Test.TestTrue("Unpack", Schema->TryGetStringField(TEXT("refPath"), RefPath));
-	Test.TestEqual("RefPath", RefPath, Expected);
-}
-
-void TestObjectRefDataInStruct(
-	FAutomationTestBase& Test, const TSharedPtr<FJsonObject>& Schema, const FString& PropertyName,
-	const FString& Expected)
-{
-	Test.TestTrue("Schema", Schema.IsValid());
-	if (!Schema.IsValid())
-	{
-		return;
+		FString RefPath;
+		Test.TestTrue("Unpack", Schema->TryGetStringField(TEXT("refPath"), RefPath));
+		Test.TestEqual("RefPath", RefPath, Expected);
 	}
 
-	const TSharedPtr<FJsonObject>* Property = nullptr;
-	Test.TestTrue("Unpack", Schema->TryGetObjectField(PropertyName, Property));
-	if (Property)
+	void TestObjectRefDataInStruct(
+		FAutomationTestBase& Test, const TSharedPtr<FJsonObject>& Schema,
+		const FString& PropertyName, const FString& Expected)
 	{
-		TestObjectRefData(Test, *Property, Expected);
+		Test.TestTrue("Schema", Schema.IsValid());
+		if (!Schema.IsValid())
+		{
+			return;
+		}
+
+		const TSharedPtr<FJsonObject>* Property = nullptr;
+		Test.TestTrue("Unpack", Schema->TryGetObjectField(PropertyName, Property));
+		if (Property)
+		{
+			TestObjectRefData(Test, *Property, Expected);
+		}
 	}
 }
 
@@ -82,11 +88,13 @@ END_DEFINE_SPEC(FToolsetJsonSpec)
 
 void FToolsetJsonSpec::Define()
 {
+	using namespace UE::ToolsetRegistry::Internal::ToolsetJson;
+
 	It("Converts UObject* property to reference schema", [this]()
 	{
 		UStruct* Struct = FToolsetJsonTest::StaticStruct();
 		FProperty* Prop = Struct->FindPropertyByName("TestObject");
-		const TSharedPtr<FJsonObject> Schema = FToolsetJson::PropertyToJsonSchema(Prop);
+		const TSharedPtr<FJsonObject> Schema = PropertyToJsonSchema(Prop);
 		TestObjectRefSchema(*this, Schema, UToolsetJsonTestObject::StaticClass(), false);
 	});
 
@@ -99,11 +107,11 @@ void FToolsetJsonSpec::Define()
 		UStruct* Struct = FToolsetJsonTest::StaticStruct();
 		FProperty* Property = Struct->FindPropertyByName("TestObject");
 
-		TSharedPtr<FJsonValue> OutJson = FToolsetJson::PropertyToJsonData(Property, &TestIn.TestObject);
+		TSharedPtr<FJsonValue> OutJson = PropertyToJsonData(Property, &TestIn.TestObject);
 		TestObjectRefData(*this, OutJson->AsObject(), TestOther->GetPathName());
 
 		FToolsetJsonTest TestOut;
-		TestTrue("Output", FToolsetJson::JsonDataToProperty(OutJson, Property, &TestOut.TestObject));
+		TestTrue("Output", JsonDataToProperty(OutJson, Property, &TestOut.TestObject));
 		TestEqual("RefCheck", TestIn.TestObject, TestOut.TestObject);
 	});
 
@@ -111,7 +119,7 @@ void FToolsetJsonSpec::Define()
 	{
 		UStruct* Struct = FToolsetJsonTest::StaticStruct();
 		FProperty* Prop = Struct->FindPropertyByName("TestObjectPath");
-		const TSharedPtr<FJsonObject> Schema = FToolsetJson::PropertyToJsonSchema(Prop);
+		const TSharedPtr<FJsonObject> Schema = PropertyToJsonSchema(Prop);
 		TestObjectRefSchema(*this, Schema, UObject::StaticClass(), false);
 	});
 
@@ -124,11 +132,11 @@ void FToolsetJsonSpec::Define()
 		UStruct* Struct = FToolsetJsonTest::StaticStruct();
 		FProperty* Property = Struct->FindPropertyByName("TestObjectPath");
 
-		TSharedPtr<FJsonValue> OutJson = FToolsetJson::PropertyToJsonData(Property, &TestIn.TestObjectPath);
+		TSharedPtr<FJsonValue> OutJson = PropertyToJsonData(Property, &TestIn.TestObjectPath);
 		TestObjectRefData(*this, OutJson->AsObject(), TestOther->GetPathName());
 
 		FToolsetJsonTest TestOut;
-		TestTrue("Output", FToolsetJson::JsonDataToProperty(OutJson, Property, &TestOut.TestObjectPath));
+		TestTrue("Output", JsonDataToProperty(OutJson, Property, &TestOut.TestObjectPath));
 		TestEqual("RefCheck", TestIn.TestObjectPath, TestOut.TestObjectPath);
 	});
 
@@ -136,7 +144,7 @@ void FToolsetJsonSpec::Define()
 	{
 		UStruct* Struct = FToolsetJsonTest::StaticStruct();
 		FProperty* Prop = Struct->FindPropertyByName("TestObjectSoft");
-		const TSharedPtr<FJsonObject> Schema = FToolsetJson::PropertyToJsonSchema(Prop);
+		const TSharedPtr<FJsonObject> Schema = PropertyToJsonSchema(Prop);
 		TestObjectRefSchema(*this, Schema, UToolsetJsonTestObject::StaticClass(), false);
 	});
 
@@ -149,11 +157,11 @@ void FToolsetJsonSpec::Define()
 		UStruct* Struct = FToolsetJsonTest::StaticStruct();
 		FProperty* Property = Struct->FindPropertyByName("TestObjectSoft");
 
-		TSharedPtr<FJsonValue> OutJson = FToolsetJson::PropertyToJsonData(Property, &TestIn.TestObjectSoft);
+		TSharedPtr<FJsonValue> OutJson = PropertyToJsonData(Property, &TestIn.TestObjectSoft);
 		TestObjectRefData(*this, OutJson->AsObject(), TestOther->GetPathName());
 
 		FToolsetJsonTest TestOut;
-		TestTrue("Output", FToolsetJson::JsonDataToProperty(OutJson, Property, &TestOut.TestObjectSoft));
+		TestTrue("Output", JsonDataToProperty(OutJson, Property, &TestOut.TestObjectSoft));
 		TestEqual("RefCheck", TestIn.TestObjectSoft, TestOut.TestObjectSoft);
 	});
 
@@ -161,7 +169,7 @@ void FToolsetJsonSpec::Define()
 	{
 		UStruct* Struct = FToolsetJsonTest::StaticStruct();
 		FProperty* Prop = Struct->FindPropertyByName("TestObjectWeak");
-		const TSharedPtr<FJsonObject> Schema = FToolsetJson::PropertyToJsonSchema(Prop);
+		const TSharedPtr<FJsonObject> Schema = PropertyToJsonSchema(Prop);
 		TestObjectRefSchema(*this, Schema, UToolsetJsonTestObject::StaticClass(), false);
 	});
 
@@ -174,11 +182,11 @@ void FToolsetJsonSpec::Define()
 		UStruct* Struct = FToolsetJsonTest::StaticStruct();
 		FProperty* Property = Struct->FindPropertyByName("TestObjectWeak");
 
-		TSharedPtr<FJsonValue> OutJson = FToolsetJson::PropertyToJsonData(Property, &TestIn.TestObjectWeak);
+		TSharedPtr<FJsonValue> OutJson = PropertyToJsonData(Property, &TestIn.TestObjectWeak);
 		TestObjectRefData(*this, OutJson->AsObject(), TestOther->GetPathName());
 
 		FToolsetJsonTest TestOut;
-		TestTrue("Output", FToolsetJson::JsonDataToProperty(OutJson, Property, &TestOut.TestObjectWeak));
+		TestTrue("Output", JsonDataToProperty(OutJson, Property, &TestOut.TestObjectWeak));
 		TestEqual("RefCheck", TestIn.TestObjectWeak, TestOut.TestObjectWeak);
 	});
 
@@ -186,7 +194,7 @@ void FToolsetJsonSpec::Define()
 	{
 		UStruct* Struct = FToolsetJsonTest::StaticStruct();
 		FProperty* Prop = Struct->FindPropertyByName("TestClass");
-		const TSharedPtr<FJsonObject> Schema = FToolsetJson::PropertyToJsonSchema(Prop);
+		const TSharedPtr<FJsonObject> Schema = PropertyToJsonSchema(Prop);
 		TestObjectRefSchema(*this, Schema, UToolsetJsonTestObject::StaticClass(), true);
 	});
 
@@ -198,11 +206,11 @@ void FToolsetJsonSpec::Define()
 		UStruct* Struct = FToolsetJsonTest::StaticStruct();
 		FProperty* Property = Struct->FindPropertyByName("TestClass");
 
-		TSharedPtr<FJsonValue> OutJson = FToolsetJson::PropertyToJsonData(Property, &TestIn.TestClass);
+		TSharedPtr<FJsonValue> OutJson = PropertyToJsonData(Property, &TestIn.TestClass);
 		TestObjectRefData(*this, OutJson->AsObject(), UToolsetJsonTestObject::StaticClass()->GetPathName());
 
 		FToolsetJsonTest TestOut;
-		TestTrue("Output", FToolsetJson::JsonDataToProperty(OutJson, Property, &TestOut.TestClass));
+		TestTrue("Output", JsonDataToProperty(OutJson, Property, &TestOut.TestClass));
 		TestEqual("RefCheck", TestIn.TestClass, TestOut.TestClass);
 	});
 
@@ -210,7 +218,7 @@ void FToolsetJsonSpec::Define()
 	{
 		UStruct* Struct = FToolsetJsonTest::StaticStruct();
 		FProperty* Prop = Struct->FindPropertyByName("TestClassPath");
-		const TSharedPtr<FJsonObject> Schema = FToolsetJson::PropertyToJsonSchema(Prop);
+		const TSharedPtr<FJsonObject> Schema = PropertyToJsonSchema(Prop);
 		TestObjectRefSchema(*this, Schema, UObject::StaticClass(), true);
 	});
 
@@ -222,17 +230,17 @@ void FToolsetJsonSpec::Define()
 		UStruct* Struct = FToolsetJsonTest::StaticStruct();
 		FProperty* Property = Struct->FindPropertyByName("TestClassPath");
 
-		TSharedPtr<FJsonValue> OutJson = FToolsetJson::PropertyToJsonData(Property, &TestIn.TestClassPath);
+		TSharedPtr<FJsonValue> OutJson = PropertyToJsonData(Property, &TestIn.TestClassPath);
 		TestObjectRefData(*this, OutJson->AsObject(), UToolsetJsonTestObject::StaticClass()->GetPathName());
 
 		FToolsetJsonTest TestOut;
-		TestTrue("Output", FToolsetJson::JsonDataToProperty(OutJson, Property, &TestOut.TestClassPath));
+		TestTrue("Output", JsonDataToProperty(OutJson, Property, &TestOut.TestClassPath));
 		TestEqual("RefCheck", TestIn.TestClassPath, TestOut.TestClassPath);
 	});
 
 	It("Converts struct to reference schema", [this]()
 	{		
-		TSharedPtr<FJsonObject> Schema = FToolsetJson::StructToJsonSchema(FToolsetJsonTest::StaticStruct());
+		TSharedPtr<FJsonObject> Schema = StructToJsonSchema(FToolsetJsonTest::StaticStruct());
 		TestTrue("Convert", Schema.IsValid());
 		if (!Schema.IsValid())
 		{
@@ -267,7 +275,7 @@ void FToolsetJsonSpec::Define()
 		TestIn.TestClass = UToolsetJsonTestObject::StaticClass();
 		TestIn.TestClassPath = UToolsetJsonTestObject::StaticClass();
 
-		TSharedPtr<FJsonObject> JsonObject = FToolsetJson::StructToJsonData(
+		TSharedPtr<FJsonObject> JsonObject = StructToJsonData(
 			FToolsetJsonTest::StaticStruct(), &TestIn);
 		for (TFieldIterator<FProperty> It(FToolsetJsonTest::StaticStruct()); It; ++It)
 		{
@@ -286,7 +294,7 @@ void FToolsetJsonSpec::Define()
 		}		
 
 		FToolsetJsonTest TestOut;
-		TestTrue("FromJson", FToolsetJson::JsonDataToStruct(
+		TestTrue("FromJson", JsonDataToStruct(
 			JsonObject.ToSharedRef(), FToolsetJsonTest::StaticStruct(), &TestOut));
 		TestEqual("RefCheck", TestIn.TestObject, TestOut.TestObject);
 		TestEqual("RefCheck", TestIn.TestObjectPath, TestOut.TestObjectPath);
@@ -304,13 +312,13 @@ void FToolsetJsonSpec::Define()
 
 		FToolsetJsonTest TestOut;
 		TSharedPtr<FJsonValue> OutJson = MakeShared<TJsonValueString<TCHAR>>(TestOther->GetPathName());
-		TestTrue("Output", FToolsetJson::JsonDataToProperty(OutJson, Property, &TestOut.TestObject));
+		TestTrue("Output", JsonDataToProperty(OutJson, Property, &TestOut.TestObject));
 		TestEqual("RefCheck", TestOut.TestObject, TestOther);
 	});
 
 	It("Filters user visible properties schema", [this]()
 	{
-		TSharedPtr<FJsonObject> ClassSchema = FToolsetJson::StructToJsonSchema(
+		TSharedPtr<FJsonObject> ClassSchema = StructToJsonSchema(
 			UToolsetJsonTestObject::StaticClass(), true);
 		TestTrue("TestData", ClassSchema.IsValid());
 		if (ClassSchema)
@@ -327,7 +335,7 @@ void FToolsetJsonSpec::Define()
 
 	It("Filters user visible properties data", [this]()
 	{
-		TSharedPtr<FJsonObject> TestData = FToolsetJson::StructToJsonData(
+		TSharedPtr<FJsonObject> TestData = StructToJsonData(
 			UToolsetJsonTestObject::StaticClass(), NewObject<UToolsetJsonTestObject>(), true);
 		TestTrue("TestData", TestData.IsValid());
 		if (TestData)
@@ -340,7 +348,7 @@ void FToolsetJsonSpec::Define()
 	It("Can read defaults", [this]()
 	{
 		UFunction* FunctionPtr = UToolsetJsonTestObject::StaticClass()->FindFunctionByName("TestDefault");
-		const TSharedPtr<FJsonObject> Schema = FToolsetJson::StructToJsonSchema(FunctionPtr);
+		const TSharedPtr<FJsonObject> Schema = StructToJsonSchema(FunctionPtr);
 		TestTrue("Schema", Schema.IsValid());
 		if (!Schema.IsValid())
 		{
diff --git a/Engine/Plugins/Experimental/ToolsetRegistry/Source/ToolsetRegistry/Private/Tests/ToolsetRegistryDelegateHandleTest.cpp b/Engine/Plugins/Experimental/ToolsetRegistry/Source/ToolsetRegistry/Private/Tests/ToolsetRegistryDelegateHandleTest.cpp
deleted file mode 100644
diff --git a/Engine/Plugins/Experimental/ToolsetRegistry/Source/ToolsetRegistry/Private/Tests/ToolsetRegistryToolCallTest.cpp b/Engine/Plugins/Experimental/ToolsetRegistry/Source/ToolsetRegistry/Private/Tests/ToolsetRegistryToolCallTest.cpp
deleted file mode 100644
diff --git a/Engine/Plugins/Experimental/ToolsetRegistry/Source/ToolsetRegistry/Private/Tests/ToolsetRegistryToolCallTest.h b/Engine/Plugins/Experimental/ToolsetRegistry/Source/ToolsetRegistry/Private/Tests/ToolsetRegistryToolCallTest.h
deleted file mode 100644
diff --git a/Engine/Plugins/Experimental/ToolsetRegistry/Source/ToolsetRegistry/Private/Tests/ToolsetRegistryToolsetTest.cpp b/Engine/Plugins/Experimental/ToolsetRegistry/Source/ToolsetRegistry/Private/Tests/ToolsetRegistryToolsetTest.cpp
deleted file mode 100644
diff --git a/Engine/Plugins/Experimental/ToolsetRegistry/Source/ToolsetRegistry/Private/Tests/ToolsetRegistryToolsetTest.h b/Engine/Plugins/Experimental/ToolsetRegistry/Source/ToolsetRegistry/Private/Tests/ToolsetRegistryToolsetTest.h
deleted file mode 100644
diff --git a/Engine/Plugins/Experimental/ToolsetRegistry/Source/ToolsetRegistry/Private/ToolsetJson.cpp b/Engine/Plugins/Experimental/ToolsetRegistry/Source/ToolsetRegistry/Private/ToolsetJson.cpp
deleted file mode 100644
diff --git a/Engine/Plugins/Experimental/ToolsetRegistry/Source/ToolsetRegistry/Private/ToolsetJson.h b/Engine/Plugins/Experimental/ToolsetRegistry/Source/ToolsetRegistry/Private/ToolsetJson.h
deleted file mode 100644
diff --git a/Engine/Plugins/Experimental/ToolsetRegistry/Source/ToolsetRegistry/Private/ToolsetLibrary.cpp b/Engine/Plugins/Experimental/ToolsetRegistry/Source/ToolsetRegistry/Private/ToolsetLibrary.cpp
--- a/Engine/Plugins/Experimental/ToolsetRegistry/Source/ToolsetRegistry/Private/ToolsetLibrary.cpp
+++ b/Engine/Plugins/Experimental/ToolsetRegistry/Source/ToolsetRegistry/Private/ToolsetLibrary.cpp
@@ -10,8 +10,8 @@
 #include "Serialization/JsonReader.h"
 #include "Serialization/JsonWriter.h"
 
-#include "ToolsetJson.h"
 #include "ToolsetRegistry/JsonConversion.h"
+#include "ToolsetRegistry/ToolsetJson.h"
 
 #include UE_INLINE_GENERATED_CPP_BY_NAME(ToolsetLibrary)
 
@@ -22,7 +22,8 @@ FString UToolsetLibraryTemp::ListClassProperties(UClass* Class)
 		return FString();
 	}
 	FString JsonProperties;
-	TSharedPtr<FJsonObject> ClassSchema = FToolsetJson::StructToJsonSchema(Class, true);
+	TSharedPtr<FJsonObject> ClassSchema =
+		UE::ToolsetRegistry::Internal::ToolsetJson::StructToJsonSchema(Class, true);
 	const TSharedPtr<FJsonObject>* ClassProperties = nullptr;
 	if (ClassSchema->TryGetObjectField(TEXT("properties"), ClassProperties))
 	{
@@ -42,8 +43,10 @@ FString UToolsetLibraryTemp::GetObjectProperties(const UObject* Object, const TA
 		FProperty* Property = Object->GetClass()->FindPropertyByName(PropertyName);
 		if(Property)
 		{
-			PropertyData->SetField(PropertyName.ToString(), FToolsetJson::PropertyToJsonData(
-				Property, Property->ContainerPtrToValuePtr<void>(Object)));
+			PropertyData->SetField(
+				PropertyName.ToString(),
+				UE::ToolsetRegistry::Internal::ToolsetJson::PropertyToJsonData(
+					Property, Property->ContainerPtrToValuePtr<void>(Object)));
 		}
 	}
 	return UE::ToolsetRegistry::Internal::JsonToString(PropertyData.ToSharedRef());
@@ -69,7 +72,8 @@ bool UToolsetLibraryTemp::SetObjectProperties(UObject* Object, const FString& Pr
 					Property, Object, EPropertyAccessChangeNotifyMode::Always, EPropertyChangeType::ValueSet);
 				PropertyAccessUtil::EmitPreChangeNotify(ChangeNotify.Get(), false);
 				void* SettingsPropPtr = Property->ContainerPtrToValuePtr<void>(Object);
-				FToolsetJson::JsonDataToProperty(PropertyData->TryGetField(Key), Property, SettingsPropPtr);
+				UE::ToolsetRegistry::Internal::ToolsetJson::JsonDataToProperty(
+					PropertyData->TryGetField(Key), Property, SettingsPropPtr);
 				PropertyAccessUtil::EmitPostChangeNotify(ChangeNotify.Get(), false);
 			}
 			else
diff --git a/Engine/Plugins/Experimental/ToolsetRegistry/Source/ToolsetRegistry/Private/ToolsetRegistry/BlueprintLibraryToolset.cpp b/Engine/Plugins/Experimental/ToolsetRegistry/Source/ToolsetRegistry/Private/ToolsetRegistry/BlueprintLibraryToolset.cpp
new file mode 100644
--- /dev/null
+++ b/Engine/Plugins/Experimental/ToolsetRegistry/Source/ToolsetRegistry/Private/ToolsetRegistry/BlueprintLibraryToolset.cpp
@@ -0,0 +1,376 @@
+// Copyright Epic Games, Inc. All Rights Reserved.
+
+#include "ToolsetRegistry/BlueprintLibraryToolset.h"
+
+#include "Async/Future.h"
+#include "Containers/UnrealString.h"
+#include "Dom/JsonObject.h"
+#include "Kismet/BlueprintFunctionLibrary.h"
+#include "Misc/CoreMisc.h"
+#include "Misc/Optional.h"
+#include "Policies/CondensedJsonPrintPolicy.h"
+#include "Serialization/JsonWriter.h"
+#include "Serialization/JsonSerializer.h"
+#include "Blueprint/BlueprintExceptionInfo.h"
+#include "ToolsetRegistryDelegateHandle.h"
+
+#include "ToolsetRegistryModule.h"
+#include "ToolsetRegistry/JsonConversion.h"
+#include "ToolsetRegistry/ObjectFunctionToolCall.h"
+#include "ToolsetRegistry/ToolsetJson.h"
+
+namespace UE::ToolsetRegistry
+{
+	namespace
+	{
+		// Register for tracking of script exceptions for the lifetime of this object.
+		class FToolCallScriptExceptionHandler
+		{
+		  public:
+			FToolCallScriptExceptionHandler()
+				: ScriptExceptionHandle(
+					FDelegateHandleRaii::Create(
+						FBlueprintCoreDelegates::OnScriptException,
+						FBlueprintCoreDelegates::OnScriptException.AddRaw(
+							this, &FToolCallScriptExceptionHandler::OnScriptException))),
+				  RuntimeScriptExceptionHandle(
+					[this](ELogVerbosity::Type Verbosity, const TCHAR* ExceptionMessage, const TCHAR* StackMessage)
+					{
+						AddToExceptionMessage(FString(ExceptionMessage));
+
+						// Also log to console to make more obvious to tool developers.
+						UE_LOG(LogToolsetRegistry, Warning, TEXT("Error running tool:\n%s"), ExceptionMessage);
+					})
+			{}
+
+			~FToolCallScriptExceptionHandler() = default;
+
+			// Track script exceptions raised by UKismetSystemLibrary::RaiseScriptException or
+			// FBlueprintCoreDelegates::ThrowScriptException via FBlueprintCoreDelegates::OnScriptException
+			// and copy into local storage.
+			void OnScriptException(
+				const UObject* ActiveObject, const FFrame& StackFrame, const FBlueprintExceptionInfo& Info)
+			{
+				TOptional<FString> ErrorMessage = GetBlueprintExceptionMessage(Info);
+				if (ErrorMessage.IsSet())
+				{
+					AddToExceptionMessage(ErrorMessage.GetValue());
+				}
+			}
+
+			FString GetException() const
+			{
+				return ScriptException;
+			}
+
+		  private:
+			TOptional<FString> GetBlueprintExceptionMessage(const FBlueprintExceptionInfo& Info)
+			{
+				TOptional<FString> ExceptionMessage;
+				switch (Info.GetType())
+				{
+					// Ignore breakpoints & tracepoints.
+					case EBlueprintExceptionType::Breakpoint:
+						// Ignore and fall through
+					case EBlueprintExceptionType::Tracepoint:
+						// Ignore and fall through
+					case EBlueprintExceptionType::WireTracepoint:
+						// Return unset message, to be ignored.
+						break;
+					case EBlueprintExceptionType::AccessViolation:
+						ExceptionMessage = Info.GetDescription().IsEmpty() ?
+							TEXT("AccessViolation") :
+							Info.GetDescription().ToString();
+						break;
+					case EBlueprintExceptionType::InfiniteLoop:
+						ExceptionMessage = Info.GetDescription().IsEmpty() ?
+							TEXT("InfiniteLoop") :
+							Info.GetDescription().ToString();
+						break;
+					case EBlueprintExceptionType::NonFatalError:
+						ExceptionMessage = Info.GetDescription().IsEmpty() ?
+							TEXT("NonFatalError") :
+							Info.GetDescription().ToString();
+						break;
+					case EBlueprintExceptionType::FatalError:
+						ExceptionMessage = Info.GetDescription().IsEmpty() ?
+							TEXT("FatalError") :
+							Info.GetDescription().ToString();
+						break;
+					case EBlueprintExceptionType::AbortExecution:
+						ExceptionMessage = Info.GetDescription().IsEmpty() ?
+							TEXT("AbortExecution") :
+							Info.GetDescription().ToString();
+						break;
+					case EBlueprintExceptionType::UserRaisedError:
+						ExceptionMessage = Info.GetDescription().IsEmpty() ?
+							TEXT("UserRaisedError") :
+							Info.GetDescription().ToString();
+						break;
+					default:
+						checkNoEntry();
+						break;
+				}
+				return ExceptionMessage;
+			}
+
+			void AddToExceptionMessage(const FString& NewMessage)
+			{
+				if (NewMessage.IsEmpty())
+				{
+					return;
+				}
+				else if (ScriptException.IsEmpty())
+				{
+					ScriptException = NewMessage;
+				}
+				else
+				{
+					ScriptException += FString::Printf(TEXT("\n%s"), *NewMessage);
+				}
+			}
+
+			// Handler for FBlueprintCoreDelegates::OnScriptException
+			FDelegateHandleRaii ScriptExceptionHandle;
+			// Handler for FScriptExceptionHandler
+			FScopedScriptExceptionHandler RuntimeScriptExceptionHandle;
+			// Raised script exception message.
+			FString ScriptException;
+		};
+	}
+
+	FBlueprintLibraryToolset::FBlueprintLibraryToolset(UClass* LibraryClass) :
+		ToolsetClass(LibraryClass)
+	{
+		GenerateToolCallObjects();
+		ToolsetName = ToolsetClass.Pin()
+			? FBlueprintLibraryToolset::GetToolsetClassName(LibraryClass) : FString();
+	}
+
+	TFuture<TValueOrError<FString, FString>> FBlueprintLibraryToolset::ExecuteTool(
+		const FString& ToolName, const FString& JsonInput)
+	{
+		FString ToolResult;
+		FString ToolError;
+		InvokeTool(ToolName, JsonInput, ToolResult, ToolError);
+		if (ToolError.IsEmpty())
+		{
+			return MakeFulfilledPromise<TValueOrError<FString, FString>>(
+				MakeValue(ToolResult)).GetFuture();
+		}
+		else
+		{
+			return MakeFulfilledPromise<TValueOrError<FString, FString>>(
+				MakeError(ToolError)).GetFuture();
+		}
+	}
+
+	FString FBlueprintLibraryToolset::GetToolsetJsonSchema() const
+	{
+		TSharedPtr<FJsonObject> Schema = GenerateSchema();
+		return UE::ToolsetRegistry ::Internal::JsonToString(Schema.ToSharedRef());
+	}
+
+	TSharedPtr<FJsonObject> FBlueprintLibraryToolset::GenerateSchema() const
+	{
+		TSharedPtr<FJsonObject> Schema = MakeShared<FJsonObject>();
+		check(Schema);
+		Schema->SetStringField(FString(TEXT("name")), GetToolsetName());
+		Schema->SetStringField(FString(TEXT("version")), GetToolsetVersion());
+		Schema->SetStringField(FString(TEXT("description")), GetToolsetDescription());
+
+		TArray<TSharedPtr<FJsonValue>> ToolsArray;
+		TArray<TObjectPtr<UFunction>> ValidTools = GetValidToolMethodsList();
+		for (UFunction* Function : ValidTools)
+		{
+			check(Function);
+			
+			TSharedPtr<FJsonObject> ToolEntry =
+				UE::ToolsetRegistry::Internal::ToolsetJson::StructToJsonSchema(Function);
+
+			// Modify tool name from "ClassName.FunctionName" to "ToolsetName.FunctionName".
+			FString ToolName = ToolEntry->GetStringField(TEXT("name"));
+			int32 LastDotIndex;
+			// Assert if there is no dot, that would mean schema generation has failed catastrophically.
+			if (ensure(ToolName.FindLastChar('.', LastDotIndex)))
+			{
+				FString FunctionName = ToolName.RightChop(LastDotIndex + 1);
+				FString NewToolName = GetToolsetName() + "." + FunctionName;
+				ToolEntry->SetStringField(TEXT("name"), *NewToolName);
+			}
+
+			ToolsArray.Add(MakeShared<FJsonValueObject>(ToolEntry));
+		}
+		Schema->SetArrayField(FString(TEXT("tools")), ToolsArray);
+		return Schema;
+	}
+
+	void FBlueprintLibraryToolset::GenerateToolCallObjects()
+	{
+		if (!Tools.IsEmpty())
+		{
+			Tools.Empty();
+		}
+		TStrongObjectPtr<UClass> Toolset = ToolsetClass.Pin();
+		if (!Toolset)
+		{
+			return;
+		}
+
+		TArray<TObjectPtr<UFunction>> ValidTools = GetValidToolMethodsList();
+		for (UFunction* Function : ValidTools)
+		{
+			check(Function);
+			const FObjectFunctionToolCall ToolCall =
+				FObjectFunctionToolCall(TNotNull<UObject*>(Toolset.Get()), TNotNull<UFunction*>(Function));
+			Tools.Add(
+				Function->GetName(), MakeShared<FObjectFunctionToolCall>(ToolCall));
+		}
+	}
+
+	void FBlueprintLibraryToolset::InvokeTool(const FString& ToolName, const FString& JsonInput, FString& OutResult, FString& OutError)
+	{
+		OutResult.Empty();
+		OutError.Empty();
+		TStrongObjectPtr<UClass> Toolset = ToolsetClass.Pin();
+		if (!Toolset)
+		{
+			OutError = FString(TEXT("Invalid Toolset Class"));
+			return;
+		}
+
+		TSharedPtr<FObjectFunctionToolCall>* ToolCallPtr = Tools.Find(ToolName);
+		if (!ToolCallPtr)
+		{
+			OutError = FString::Printf(TEXT("Unknown tool %s"), *ToolName);
+			return;
+		}
+		TSharedPtr<FObjectFunctionToolCall> ToolCall = *ToolCallPtr;
+		check(ToolCall);
+
+		FObjectFunctionToolCall::FFunctionInputParamsJson JsonParams;
+		JsonParams.Set<FString>(JsonInput);
+
+		// Capture script errors raised by the tool during execution. Note: script errors raised by
+		// tools running outside of the game thread will not be captured correctly.
+		FToolCallScriptExceptionHandler ExceptionHandler;
+		ensureAlways(IsInGameThread());
+
+		FObjectFunctionToolCall::FToolCallResult ToolCallResult = ToolCall->Execute(JsonParams);
+		if (ToolCallResult.HasError())
+		{
+			OutError = ToolCallResult.StealError();
+			return;
+		}
+		else if (FString ScriptException = ExceptionHandler.GetException();
+		         !ScriptException.IsEmpty())
+		{
+			// If the tool raised an error, return that.
+			OutError = ScriptException;
+			return;
+		}
+		check(ToolCallResult.HasValue());
+		TSharedPtr<FJsonObject> ResultValue = ToolCallResult.StealValue();
+		if (!ResultValue.IsValid())
+		{
+			OutError = FString::Printf(TEXT("Unable to serialize result from executing tool %s"), *ToolName);
+			return;
+		}
+		OutResult = UE::ToolsetRegistry::Internal::JsonToString(ResultValue.ToSharedRef());
+	}
+
+	TArray<FString> FBlueprintLibraryToolset::GetToolNames()
+	{
+		TArray<TObjectPtr<UFunction>> ValidTools = GetValidToolMethodsList();
+		TArray<FString> ToolNames;
+		for (TObjectPtr<UFunction> FuncPtr : ValidTools)
+		{
+			ToolNames.Add(FuncPtr->GetName());
+		}
+		return ToolNames;
+	}
+
+	bool FBlueprintLibraryToolset::HasValidTools()
+	{
+		return !Tools.IsEmpty();
+	}
+
+	TArray<TObjectPtr<UFunction>> FBlueprintLibraryToolset::GetValidToolMethodsList() const
+	{
+		TArray<TObjectPtr<UFunction>> ValidTools;
+		TStrongObjectPtr<UClass> Toolset = ToolsetClass.Pin();
+		if (!Toolset)
+		{
+			return ValidTools;
+		}
+
+		for (TFieldIterator<UFunction> FuncIt(Toolset.Get(), EFieldIterationFlags::None); FuncIt; ++FuncIt)
+		{
+			TObjectPtr<UFunction> Function = *FuncIt;
+			if (Function)
+			{
+				bool bValidFunction = true;
+
+				// Skip functions marked as for internal use only.
+				if (Function->HasMetaData(TEXT("BlueprintInternalUseOnly")))
+				{
+					bValidFunction = false;
+				}
+
+				// Skip functions unless they are public and static.
+				if (!Function->HasAllFunctionFlags(FUNC_Public | FUNC_Static))
+				{
+					bValidFunction = false;
+				}
+
+				if (bValidFunction)
+				{
+					ValidTools.Add(Function);
+				}
+			}
+		}
+		return ValidTools;
+	}
+
+	/*static*/ FString FBlueprintLibraryToolset::GetToolsetClassName(UClass* InClass)
+	{
+		FString Name;
+		if (InClass == nullptr)
+		{
+			return Name;
+		}
+		FString ClassName = InClass->GetName();
+		Name = ClassName;
+		check(InClass->GetOuter());
+		FString PackageName = InClass->GetOuter()->GetName();
+		int32 LastSlashIndex;
+		FString QualifierName;
+		if (ensure(PackageName.FindLastChar('/', LastSlashIndex)))
+		{
+			// Convert package name like "/Script/ToolsetRegistry" to a shorter qualifier like
+			// "ToolsetRegistry".
+			QualifierName = PackageName.RightChop(LastSlashIndex + 1);
+			if (QualifierName.EndsWith(FString("_PY")))
+			{
+				// Special handling for Python packages: convert package name like
+				// "/ToolsetRegistry/Python/toolset_registry/tests/demo_toolset_PY" as generated by
+				// PyUtil::GetGeneratedTypeOuterAndName, to "toolset_registry.tests.demo_toolset".
+				TArray<FString> PackageParts;
+				PackageName.ParseIntoArray(PackageParts, TEXT("/"));
+				int32 PythonPartIndex;
+				if (PackageParts.FindLast(TEXT("Python"), PythonPartIndex))
+				{
+					PackageParts.RemoveAt(0, PythonPartIndex + 1);
+				}
+				if (int32 ChopIndex = PackageParts.Last().Find(TEXT("_PY")); ChopIndex != INDEX_NONE)
+				{
+					PackageParts.Last() = PackageParts.Last().Left(ChopIndex);
+				}
+				QualifierName = FString::Join(PackageParts, TEXT("."));
+			}
+
+			Name = QualifierName + "." + ClassName;
+		}
+		return Name;
+	}
+}
diff --git a/Engine/Plugins/Experimental/ToolsetRegistry/Source/ToolsetRegistry/Private/ToolsetRegistry/ObjectFunctionToolCall.cpp b/Engine/Plugins/Experimental/ToolsetRegistry/Source/ToolsetRegistry/Private/ToolsetRegistry/ObjectFunctionToolCall.cpp
new file mode 100644
--- /dev/null
+++ b/Engine/Plugins/Experimental/ToolsetRegistry/Source/ToolsetRegistry/Private/ToolsetRegistry/ObjectFunctionToolCall.cpp
@@ -0,0 +1,474 @@
+// Copyright Epic Games, Inc. All Rights Reserved.
+
+#include "ToolsetRegistry/ObjectFunctionToolCall.h"
+
+#include "Async/Async.h"
+#include "Engine/Engine.h"
+#include "Engine/World.h"
+#include "Engine/EngineTypes.h"
+#include "Serialization/JsonReader.h"
+#include "Serialization/JsonSerializer.h"
+#include "ToolsetJson.h"
+#include "UObject/FieldIterator.h"
+
+#include "ToolsetRegistry/JsonConversion.h"
+
+namespace UE::ToolsetRegistry
+{
+	using namespace UE::ToolsetRegistry::Internal;
+
+	namespace
+	{
+		TSharedRef<FJsonObject> FunctionInputParamsJsonToValidJsonObject(
+			const TOptional<FObjectFunctionToolCall::FFunctionInputParamsJson>&
+				FunctionInputParamsJson)
+		{
+			if (FunctionInputParamsJson.IsSet())
+			{
+				if (FunctionInputParamsJson->IsType<TSharedPtr<FJsonObject>>())
+				{
+					// Return a copy of the JSON object, since it may be altered, and we 
+					// don't want to alter the original.
+					const TSharedPtr<FJsonObject> OriginalJsonObject =
+						FunctionInputParamsJson->Get<TSharedPtr<FJsonObject>>();
+					return OriginalJsonObject
+						? MakeShared<FJsonObject>(*OriginalJsonObject)
+						: MakeShared<FJsonObject>();
+				}
+				else
+				{
+					return JsonObjectOrEmpty(
+						JsonStringToJsonObject(FunctionInputParamsJson->Get<FString>()));
+				}
+			}
+
+			return MakeShared<FJsonObject>();
+		}
+
+		UWorld* GetWorld()
+		{
+			if (GEngine)
+			{
+				for (const FWorldContext& WorldContext : GEngine->GetWorldContexts())
+				{
+					if (WorldContext.WorldType == EWorldType::Game ||
+						WorldContext.WorldType == EWorldType::PIE)
+					{
+						return WorldContext.World();
+					}
+				}
+			}
+
+			return nullptr;
+		}
+
+		void MaybeSetWorldContextPropertyInFrameMemory(const UFunction* Function,
+			uint8* FrameMemory)
+		{
+			const FName WorldContextKey(TEXT("WorldContext"));
+			const FName WorldContextParamName =
+				Function->HasMetaData(WorldContextKey) ?
+				FName(*Function->GetMetaData(WorldContextKey)) :
+				FName(TEXT("WorldContextObject"));
+
+			if (FProperty* WorldContextProperty =
+				Function->FindPropertyByName(WorldContextParamName))
+			{
+				UWorld* World = GetWorld();
+				FObjectPropertyBase* ObjectProperty =
+					CastField<FObjectPropertyBase>(WorldContextProperty);
+				if (ObjectProperty && World)
+				{
+					ObjectProperty->SetObjectPropertyValue_InContainer(FrameMemory, World);
+				}
+			}
+		}
+
+		/**
+		  * Checks whether all the input parameters required by the function schema JSON are
+		  * provided in the function input params JSON.
+		  * @param FunctionDebugName Name of function, for debugging and error messages.
+		  * @param FunctionInputParamsJsonObject The JSON option containing the function call
+		  *     params.
+		  * @param FunctionSchemaInputSchemaJsonObject The 'inputSchema' sub-JSON in the main
+		  *     function schema JSON.
+		  *	@param OutErrorString Output error string if a required param is missing, empty
+		  *     otherwise.
+		  *	@return Whether required params are all supplied.
+		  */
+		bool AreAllRequiredFunctionInputParamsProvided(
+			const FString& FunctionDebugName,
+			const TSharedRef<FJsonObject>& FunctionInputParamsJsonObject,
+			const TSharedRef<FJsonObject>& FunctionSchemaInputSchemaJsonObject,
+			FString* OutErrorString = nullptr)
+		{
+			// For each required param name, make sure we have a field for it in the function 
+			// input params.
+			{
+				TArray<FString> RequiredParamNames;
+				FunctionSchemaInputSchemaJsonObject->TryGetStringArrayField(
+					TEXT("required"), RequiredParamNames);
+
+				if (!RequiredParamNames.IsEmpty() &&
+					FunctionInputParamsJsonObject->Values.IsEmpty())
+				{
+					// Caller error. Missing but expected params in incoming function 
+					// input params JSON. 
+					if (OutErrorString)
+					{
+						*OutErrorString = FString::Printf(TEXT(
+							"Function \"%s\", "
+							"input params are required by the function input schema Json, "
+							"but incoming function input params Json is empty.\n"
+							"Function schema Json -\n"
+							"%s"),
+							*FunctionDebugName,
+							*JsonToString(FunctionSchemaInputSchemaJsonObject));
+					}
+
+					return false;
+				}
+
+				for (const FString& RequiredParamName : RequiredParamNames)
+				{
+					if (!FunctionInputParamsJsonObject->HasField(*RequiredParamName))
+					{
+						// Caller error. A missing but expected param in incoming function 
+						// input params JSON.
+						if (OutErrorString)
+						{
+							*OutErrorString = FString::Printf(TEXT(
+								"Function \"%s\", "
+								"input param \"%s\" is required by the function input schema "
+								"Json, but is missing from the incoming function input params "
+								"Json.\n"
+								"Function schema Json -\n"
+								"%s\n"
+								"Function input params Json -\n"
+								"%s"),
+								*FunctionDebugName,
+								*RequiredParamName,
+								*JsonToString(FunctionSchemaInputSchemaJsonObject),
+								*JsonToString(FunctionInputParamsJsonObject));
+						}
+
+						return false;
+					}
+				}
+			}
+
+			return true;
+		}
+
+		/**
+		  * For each input schema property, if we don't have a field for it in the
+		  * function input params, then add it to the function input params, and
+		  * give it the same value as the default value of the corresponding function
+		  * input schema property.
+		  * (Note - Checked that we have all the required params first, so we know those
+		  * will all be in the function input params JSON. Then this ends up only handling
+		  * params that are both non-required and missing from the function input param
+		  * JSON. To apply default values, we need a 'default' value to exist for each
+		  * corresponding param object in the function input schema.)
+		  * @param FunctionDebugName Name of the function, used for error message if needed.
+		  * @param FunctionInputParamsJsonObject The JSON object containing the function call
+		  *     params.
+		  * @param FunctionSchemaInputSchemaJsonObject The 'inputSchema' sub-JSON in the main
+		  *     function schema JSON.
+		  * @param OutErrorString Error message, if function fails.
+		  * @return Whether this function could provide defaults for all missing function input params.
+		 */
+		bool ProvideDefaultsForMissingFunctionInputParams(
+			const FString& FunctionDebugName,
+			const TSharedRef<FJsonObject>& FunctionInputParamsJsonObject,
+			const TSharedRef<FJsonObject>& FunctionSchemaInputSchemaJsonObject,
+			FString* OutErrorString = nullptr)
+		{
+			if (const TSharedPtr<FJsonObject>* InputPropertiesJsonObject;
+				FunctionSchemaInputSchemaJsonObject->TryGetObjectField(
+					TEXT("properties"), InputPropertiesJsonObject))
+			{
+				for (const TPair<FString, TSharedPtr<FJsonValue>>& Pair :
+					(*InputPropertiesJsonObject)->Values)
+				{
+					const FString& InputParamPropertyName = Pair.Key;
+
+					const TSharedPtr<FJsonObject>& InputPropertyJsonObject =
+						Pair.Value->AsObject();
+					if (!InputPropertyJsonObject)
+					{
+						// Caller error. Missing but expected Json for an input param 
+						// in incoming function input params JSON. 
+						if (OutErrorString)
+						{
+							*OutErrorString = FString::Printf(TEXT(
+								"Function \"%s\", "
+								"input param \"%s\" has no Json object field in input params.\n"
+								"Function input params Json -\n"
+								"%s"),
+								*FunctionDebugName,
+								*InputParamPropertyName,
+								*JsonToString(FunctionInputParamsJsonObject));
+						}
+
+						return false;
+					}
+
+					if (!FunctionInputParamsJsonObject->HasField(*InputParamPropertyName))
+					{
+						const TSharedPtr<FJsonValue> InputPropertyDefaultJsonValue =
+							InputPropertyJsonObject->TryGetField(TEXT("default"));
+
+						if (!InputPropertyDefaultJsonValue)
+						{
+							// Caller error. Missing but expected Json for an input param's 
+							// default in incoming function input params JSON. 
+							if (OutErrorString)
+							{
+								*OutErrorString = FString::Printf(TEXT(
+									"Function \"%s\", "
+									"input param \"%s\" needs a default value."
+									"Function input params Json -\n"
+									"%s"),
+									*FunctionDebugName,
+									*InputParamPropertyName,
+									*JsonToString(FunctionInputParamsJsonObject));
+							}
+
+							return false;
+						}
+
+						FunctionInputParamsJsonObject->SetField(*InputParamPropertyName,
+							InputPropertyDefaultJsonValue);
+					}
+				}
+			}
+
+			return true;
+		}
+	}
+
+	//
+	// UE::ToolsetRegistry::FObjectFunctionToolCall
+	//
+
+	FObjectFunctionToolCall::FObjectFunctionToolCall(
+		TNotNull<UObject*> InInstanceObject, TNotNull<UFunction*> InFunction) :
+		InstanceObject(InInstanceObject),
+		Function(InFunction)
+	{
+		FunctionSchemaJsonObject = ToolsetJson::StructToJsonSchema(Function);
+		check(FunctionSchemaJsonObject);
+	}
+
+	/*static*/ FObjectFunctionToolCall FObjectFunctionToolCall::ConstructWithNamedClassFunction(
+		TNotNull<UObject*> InInstanceObject, const FName& InFunctionName)
+	{
+		check(!InFunctionName.IsNone());
+
+		UFunction* Function = InInstanceObject->GetClass()->FindFunctionByName(InFunctionName);
+		check(Function);
+
+		return FObjectFunctionToolCall(InInstanceObject, Function);
+	}
+
+	FObjectFunctionToolCall::FToolCallResult FObjectFunctionToolCall::Execute(
+		const TOptional<FFunctionInputParamsJson>& FunctionInputParamsJson) const
+	{
+		// Get function input params JSON as a JSON object.
+		TSharedPtr<FJsonObject> FunctionInputParamsJsonObject;
+		if (const TValueOrError<const TSharedRef<FJsonObject>, const FString> JsonObjectOrError =
+			BuildValidFunctionInputParamsJsonObject(FunctionInputParamsJson);
+			JsonObjectOrError.HasError())
+		{
+			return MakeError(JsonObjectOrError.GetError());
+		}
+		else
+		{
+			FunctionInputParamsJsonObject = JsonObjectOrError.GetValue();
+		}
+
+		// Allocate the frame memory.
+		uint8* FrameMemory;
+		{
+			// Invoke must use properties size, because it's allocating locals, 
+			// i.e. for script functions.
+			const int32 FrameSize = Function->PropertiesSize;
+			FrameMemory = static_cast<uint8*>(FMemory_Alloca_Aligned(FrameSize,
+				Function->GetMinAlignment()));
+			FMemory::Memzero(FrameMemory, FrameSize);
+		}
+
+		// We need to manually allocate each property in the frame memory.
+		for (TFieldIterator<FProperty> It(Function); It && It->HasAnyPropertyFlags(CPF_Parm); ++It)
+		{
+			if (const FProperty* ParameterProp = *It;
+				!ParameterProp->HasAnyPropertyFlags(CPF_ZeroConstructor))
+			{
+				ParameterProp->InitializeValue_InContainer(FrameMemory);
+			}
+		}
+
+		// When we exit, we need to manually destroy each property in the frame memory.
+		ON_SCOPE_EXIT
+		{
+			// Invoke will destroy locals, but caller owns parameters. This should also 
+			// destroy the result parameter.
+			for (TFieldIterator<FProperty> It(Function);
+				It && It->HasAnyPropertyFlags(CPF_Parm); ++It)
+			{
+				It->DestroyValue_InContainer(FrameMemory);
+			}
+		};
+
+		// Copy the function input params memory into the function's frame memory.
+		if (FrameMemory &&
+			!ToolsetJson::JsonDataToStruct(
+				FunctionInputParamsJsonObject.ToSharedRef(), Function, FrameMemory))
+		{
+			// Caller error. Ill-formed incoming function input params JSON. 
+			return MakeError(FString::Printf(TEXT(
+				"Function \"%s\", "
+				"could not convert incoming function input params Json to a UStruct. "
+				"Incoming function input params Json may be incorrect or contain unsupported "
+				"types.\n"
+				"Function input params Json -\n"
+				"%s"),
+				*Function->GetName(),
+				*JsonToString(FunctionInputParamsJsonObject.ToSharedRef())));
+		}
+
+		// In case this is a Blueprint function that has a world context object parameter, 
+		// give that parameter (in the frame memory) the World.
+		MaybeSetWorldContextPropertyInFrameMemory(Function, FrameMemory);
+
+		// Invoke function.
+		// TODO: A lot of this logic is duplicated from ScriptCore and should be ported to use existing function invocation logic.
+		FFrame StackFrame(InstanceObject, Function, FrameMemory, nullptr,
+			Function->ChildProperties);
+
+		// If the function has out parameters, fill the stack frame's out parameter info with the info for those params 
+		if (Function->HasAnyFunctionFlags(FUNC_HasOutParms))
+		{
+			FOutParmRec** LastOut = &StackFrame.OutParms;
+			for (FProperty* Property = (FProperty*)(Function->ChildProperties); Property && (Property->PropertyFlags & (CPF_Parm)) == CPF_Parm; Property = (FProperty*)Property->Next)
+			{
+				// this is used for optional parameters - the destination address for out parameter values is the address of the calling function
+				// so we'll need to know which address to use if we need to evaluate the default parm value expression located in the new function's
+				// bytecode
+				if (Property->HasAnyPropertyFlags(CPF_OutParm))
+				{
+					CA_SUPPRESS(6263)
+					FOutParmRec* Out = (FOutParmRec*)FMemory_Alloca(sizeof(FOutParmRec));
+					// set the address and property in the out param info
+					// note that since C++ doesn't support "optional out" we can ignore that here
+					Out->PropAddr = Property->ContainerPtrToValuePtr<uint8>(FrameMemory);
+					Out->Property = Property;
+
+					// add the new out param info to the stack frame's linked list
+					if (*LastOut)
+					{
+						(*LastOut)->NextOutParm = Out;
+						LastOut = &(*LastOut)->NextOutParm;
+					}
+					else
+					{
+						*LastOut = Out;
+					}
+				}
+			}
+
+			// set the next pointer of the last item to NULL to mark the end of the list
+			if (*LastOut)
+			{
+				(*LastOut)->NextOutParm = NULL;
+			}
+		}
+
+		const bool bHasReturnParam = (Function->ReturnValueOffset != MAX_uint16);
+		uint8* ReturnValueAddress = (bHasReturnParam ?
+			(FrameMemory + Function->ReturnValueOffset) : nullptr);
+		Function->Invoke(InstanceObject, StackFrame, ReturnValueAddress);
+
+		// Get the return value, if we have one.
+		TSharedPtr<FJsonValue> ReturnJsonValue = Function->GetReturnProperty() ?
+			ToolsetJson::PropertyToJsonData(Function->GetReturnProperty(), ReturnValueAddress) :
+			MakeShared<FJsonValueNull>();
+
+		FString ErrorMessage;
+		TSharedPtr<FJsonObject> ReturnJsonObject = MakeReturnObject(ReturnJsonValue, ErrorMessage);
+
+		if (!ErrorMessage.IsEmpty())
+		{
+			return MakeError(ErrorMessage);
+		}
+		else
+		{
+			return MakeValue(ReturnJsonObject);
+		}
+	}
+
+	TSharedPtr<FJsonObject> FObjectFunctionToolCall::MakeReturnObject(
+		const TSharedPtr<FJsonValue>& InReturnValue, FString& OutErrorMessage) const
+	{
+		// If the value to return is invalid, return an error message instead of the return object.
+		OutErrorMessage.Empty();
+		if (!InReturnValue.IsValid())
+		{
+			OutErrorMessage = FString::Printf(
+				TEXT("Function \"%s\": could not convert function return value to Json value."),
+				*Function->GetName());
+			return nullptr;
+		}
+
+		// According to the output schema: all returned values are required to be wrapped in a "returnValue"
+		// object.
+		TSharedPtr<FJsonObject> ReturnJsonObject = MakeShared<FJsonObject>();
+		ReturnJsonObject->SetField(TEXT("returnValue"), InReturnValue);
+		return ReturnJsonObject;
+	}
+
+	TValueOrError<const TSharedRef<FJsonObject>, const FString>
+		FObjectFunctionToolCall::BuildValidFunctionInputParamsJsonObject(
+			const TOptional<FFunctionInputParamsJson>& FunctionInputParamsJson) const
+	{
+		// Get incoming function params as JSON object.
+		const TSharedRef<FJsonObject> FunctionInputParamsJsonObject =
+			FunctionInputParamsJsonToValidJsonObject(FunctionInputParamsJson);
+
+		// With function schema JSON -
+		// (1) Validate the function input JSON against the function schema's input schema JSON.
+		// (2) Set any missing params in the function input params JSON to the defaults from 
+		//		function schema's input schema JSON.
+		if (const TSharedPtr<FJsonObject>* FunctionSchemaInputSchemaJsonObjectPtr;
+			FunctionSchemaJsonObject->TryGetObjectField(
+				TEXT("inputSchema"), FunctionSchemaInputSchemaJsonObjectPtr))
+		{
+			// TODO - We should ultimately be fully validating the function input params JSON 
+			// against the function schema's input schema JSON instead of this. This isn't 
+			// validating type compatibility, only checking existence of params.
+			if (FString ErrorString;
+				!AreAllRequiredFunctionInputParamsProvided(
+					Function->GetName(),
+					FunctionInputParamsJsonObject,
+					FunctionSchemaInputSchemaJsonObjectPtr->ToSharedRef(),
+					&ErrorString))
+			{
+				return MakeError(ErrorString);
+			}
+
+			if (FString ErrorString;
+				!ProvideDefaultsForMissingFunctionInputParams(
+					Function->GetName(),
+					FunctionInputParamsJsonObject,
+					FunctionSchemaInputSchemaJsonObjectPtr->ToSharedRef(),
+					&ErrorString))
+			{
+				return MakeError(ErrorString);
+			}
+		}
+
+		return MakeValue(FunctionInputParamsJsonObject);
+	}
+}
+
diff --git a/Engine/Plugins/Experimental/ToolsetRegistry/Source/ToolsetRegistry/Private/ToolsetRegistry/ObjectFunctionToolCall.h b/Engine/Plugins/Experimental/ToolsetRegistry/Source/ToolsetRegistry/Private/ToolsetRegistry/ObjectFunctionToolCall.h
new file mode 100644
--- /dev/null
+++ b/Engine/Plugins/Experimental/ToolsetRegistry/Source/ToolsetRegistry/Private/ToolsetRegistry/ObjectFunctionToolCall.h
@@ -0,0 +1,81 @@
+// Copyright Epic Games, Inc. All Rights Reserved.
+
+#pragma once
+
+#include "Templates/SharedPointer.h"
+#include "Misc/NotNull.h"
+#include "Misc/TVariant.h"
+#include "Templates/ValueOrError.h"
+#include "UObject/Class.h"
+#include "Dom/JsonObject.h"
+
+//
+// UE::ToolsetRegistry::FObjectFunctionToolCall
+//
+
+namespace UE::ToolsetRegistry
+{
+	class FObjectFunctionToolCall
+	{
+	public:
+
+		/**
+		 * Function input params.
+		 * Holds either a JSON object or a JSON string.
+		 */
+		using FFunctionInputParamsJson = TVariant<TSharedPtr<FJsonObject>, FString>;
+
+		/**
+		 * Tool call output result.
+		 * Either a JSON object or a string describing an error.
+		 */
+		using FToolCallResult = TValueOrError<TSharedPtr<FJsonObject>, FString>;
+		
+		/**
+		 * Constructs with a UFunction on a UObject instance. 
+		* @param InInstanceObject UObject containing the function.
+		 * @param InFunction UFunction contained in the UObject.
+		 */
+		FObjectFunctionToolCall(
+			TNotNull<UObject*> InInstanceObject, TNotNull<UFunction*> InFunction);
+		
+		/**
+		 * Constructs with a UFunction on a UObject instance, but using the UObject's UClass 
+		 * where the function can be found by name.
+		 * @param InInstanceObject UObject containing the function.
+		 * @param InFunctionName UFunction name contained in the UObject.
+		 * @return A function call object.
+		 */
+		static FObjectFunctionToolCall ConstructWithNamedClassFunction(
+			TNotNull<UObject*> InInstanceObject, const FName& InFunctionName);
+		
+		/**
+		 * Given a JSON containing input parameters for a tool call, execute the function with those 
+		 * input parameters.
+		 * @param FunctionInputParamsJson The optional JSON object or string containing the function 
+		 *		parameters. If the function this class was initialized with has required params, then
+		 *		this must be supplied with JSON that provides those required params. Otherwise, an
+		 *		error will result. If the function this class was initialized with has no required 
+		 *		params, then this does not need to be supplied.
+		 * @return Tool call result. (See above). When tool result does not have an error -
+ 		 *		Returned objects will either be (1) the object itself, or (2) a value named 
+ 		 *		'returnValue' inside an object. And this may be null if (1) function does not 
+ 		 *		return a value, or (2) if function does return a value and that return value 
+ 		 *		is null.
+		 */
+		FToolCallResult Execute(const TOptional<FFunctionInputParamsJson>& FunctionInputParamsJson = 
+			TOptional<FFunctionInputParamsJson>()) const;
+
+	private:
+		TValueOrError<const TSharedRef<FJsonObject>, const FString> 
+		BuildValidFunctionInputParamsJsonObject(
+			const TOptional<FFunctionInputParamsJson>& FunctionInputParamsJson) const;
+
+		TSharedPtr<FJsonObject> MakeReturnObject(
+			const TSharedPtr<FJsonValue>& InReturnValue, FString& OutErrorMessage) const;
+
+		TNotNull<UObject*> InstanceObject;
+		TNotNull<UFunction*> Function;
+		TSharedPtr<FJsonObject> FunctionSchemaJsonObject;
+	};
+}
diff --git a/Engine/Plugins/Experimental/ToolsetRegistry/Source/ToolsetRegistry/Private/ToolsetRegistry/ToolsetJson.cpp b/Engine/Plugins/Experimental/ToolsetRegistry/Source/ToolsetRegistry/Private/ToolsetRegistry/ToolsetJson.cpp
new file mode 100644
--- /dev/null
+++ b/Engine/Plugins/Experimental/ToolsetRegistry/Source/ToolsetRegistry/Private/ToolsetRegistry/ToolsetJson.cpp
@@ -0,0 +1,291 @@
+// Copyright Epic Games, Inc. All Rights Reserved.
+
+#include "ToolsetJson.h"
+
+#include "Editor.h"
+#include "JsonObjectConverter.h"
+#include "JsonObjectStructInterface.h"
+#include "JsonSchema/JsonSchemaGenerator.h"
+#include "Modules/ModuleManager.h"
+#include "Misc/AssertionMacros.h"
+
+namespace
+{
+	// Use RequiredModules::EnsureLoaded() to lazily load required modules.
+	class RequiredModules
+	{
+	private:
+		RequiredModules()
+		{
+			FModuleManager::Get().LoadModuleChecked(TEXT("JsonUtilitiesEditor"));
+		}
+
+	public:
+		static void EnsureLoaded()
+		{
+			// Construct a singleton to ensure required modules are loaded for methods
+			// in this class.
+			static RequiredModules Load;
+		}
+	};
+
+	EPropertyFlags GetCheckFlags(bool UserVisiblePropertiesOnly)
+	{
+		EPropertyFlags CheckFlags = CPF_None;
+		if (UserVisiblePropertiesOnly)
+		{
+			CheckFlags = CPF_BlueprintVisible | CPF_BlueprintAssignable;
+			if (GEditor)
+			{
+				CheckFlags |= CPF_Edit;
+			}
+		}
+		return CheckFlags;
+	}
+
+	bool IsReferenceLikeProperty(TNotNull<const FProperty*> Property)
+	{
+		if (CastField<FObjectProperty>(Property) &&
+			!Property->HasAnyPropertyFlags(CPF_PersistentInstance))
+		{
+			return true;
+		}
+		else if (
+			CastField<FSoftObjectProperty>(Property) ||
+			CastField<FWeakObjectProperty>(Property))
+		{
+			return true;
+		}
+		else if (const FStructProperty* StructProperty = CastField<FStructProperty>(Property))
+		{
+			const UScriptStruct* ScriptStruct = StructProperty->Struct;
+			if (ScriptStruct == FSoftObjectPath::StaticStruct() ||
+				ScriptStruct == FSoftClassPath::StaticStruct())
+			{
+				return true;
+			}
+		}
+
+		return false;
+	}
+
+	const FStructProperty* AsColorProperty(TNotNull<const FProperty*> Property)
+	{
+		if (const FStructProperty* StructProperty = CastField<FStructProperty>(Property))
+		{
+			if (StructProperty->Struct->GetFName() == NAME_Color)
+			{
+				return StructProperty;
+			}
+		}
+		return nullptr;
+	}
+}
+
+namespace UE::ToolsetRegistry::Internal::ToolsetJson
+{
+
+// This method generates customizes the default JSON schema generation logic in a way that is easier for an
+// LLM to use reliably. The primary motivation for this method is converting UObject* and UClass* from
+// "string" types to a distinct reference type.
+bool CustomJsonSchema(
+	const FProperty* Property, const FString& ParameterDefaultString, const TSharedRef<FJsonObject>& OutputSchema)
+{
+	RequiredModules::EnsureLoaded();
+	if (IsReferenceLikeProperty(Property))
+	{
+		if (ParameterDefaultString.IsEmpty())
+		{
+			FJsonObject::Duplicate(
+				FJsonSchemaGenerator::UStructToJsonSchemaObject<FToolsetReference>(), OutputSchema);
+			UClass* ClassType = nullptr;
+			bool IsClass = false;
+			if (const FClassProperty* ClassProperty = CastField<FClassProperty>(Property))
+			{
+				ClassType = ClassProperty->MetaClass;
+				IsClass = true;
+			}
+			else if (const FObjectPropertyBase* ObjectProperty = CastField<FObjectPropertyBase>(Property))
+			{
+				ClassType = ObjectProperty->PropertyClass;
+			}
+			else if (const FStructProperty* StructProperty = CastField<FStructProperty>(Property))
+			{
+				const UScriptStruct* ScriptStruct = StructProperty->Struct;
+				if (ScriptStruct == FSoftObjectPath::StaticStruct())
+				{
+					ClassType = UObject::StaticClass();
+				}
+				else if (ScriptStruct == FSoftClassPath::StaticStruct())
+				{
+					ClassType = UObject::StaticClass();
+					IsClass = true;
+				}
+			}
+			check(ClassType);
+			FString TypeInfo = IsClass ? FString(TEXT("Class@")) : FString();
+			TypeInfo += ClassType->GetPathName();
+			OutputSchema->SetStringField(TEXT("title"), TypeInfo);
+		}
+		else
+		{
+			if (ParameterDefaultString == FString(TEXT("None")))
+			{
+				OutputSchema->SetField(TEXT("default"), MakeShared<FJsonValueNull>());
+			}
+			else
+			{
+				FToolsetReference Reference;
+				Reference.RefPath = ParameterDefaultString;
+				OutputSchema->SetObjectField(
+					TEXT("default"),
+					FJsonObjectConverter::UStructToJsonObject(Reference));
+			}
+		}
+		return true;
+	}
+	return false;
+}
+
+// This method customizes the default JSON data import logic to match the custom schema and
+// data export logic in this file.
+bool CustomImport(const TSharedPtr<FJsonValue>& JsonValue, FProperty* Property, void* Value)
+{
+	ensureAlways(JsonValue.IsValid());
+	RequiredModules::EnsureLoaded();
+	if (IsReferenceLikeProperty(Property))
+	{
+		FString RefPath;
+		if (JsonValue->Type == EJson::Object)
+		{						
+			JsonValue->AsObject()->TryGetStringField(TEXT("refPath"), RefPath);
+		}
+		else if (JsonValue->Type == EJson::String)
+		{
+			RefPath = JsonValue->AsString();
+		}
+		if (!RefPath.IsEmpty() && Property->ImportText_Direct(*RefPath, Value, nullptr, 0))
+		{
+			return true;
+		}
+	}
+	return false;
+}
+
+// This method customizes the default JSON data export behavior to align with the
+// custom schema generated by CustomJsonSchema().
+TSharedPtr<FJsonValue> CustomExport(FProperty* Property, const void* Value)
+{
+	RequiredModules::EnsureLoaded();
+	if (IsReferenceLikeProperty(Property))
+	{
+		FToolsetReference Reference;
+		if (FObjectPropertyBase* ObjectProperty = CastField<FObjectPropertyBase>(Property))
+		{
+			if (const UObject* Object = ObjectProperty->GetObjectPropertyValue(Value))
+			{
+				Reference.RefPath = Object->GetPathName();
+			}
+			else
+			{
+				return nullptr;
+			}
+		}		
+		else if (const FStructProperty* StructProperty = CastField<FStructProperty>(Property))
+		{
+			const UScriptStruct* ScriptStruct = StructProperty->Struct;
+			if (ScriptStruct == FSoftObjectPath::StaticStruct())
+			{
+				const FSoftObjectPath* SoftObjectPath = reinterpret_cast<const FSoftObjectPath*>(Value);
+				Reference.RefPath = SoftObjectPath->ToString();
+			}
+			else if (ScriptStruct == FSoftClassPath::StaticStruct())
+			{
+				const FSoftClassPath* SoftClassPath = reinterpret_cast<const FSoftClassPath*>(Value);
+				Reference.RefPath = SoftClassPath->ToString();
+			}
+			else
+			{
+				check(false);
+			}
+		}
+		else
+		{
+			check(false);
+		}
+		return MakeShared<FJsonValueObject>(FJsonObjectConverter::UStructToJsonObject(Reference));
+	}
+	
+	return nullptr;
+}
+
+TSharedPtr<FJsonObject> PropertyToJsonSchema(TNotNull<const FProperty*> Property)
+{
+	RequiredModules::EnsureLoaded();
+	FJsonSchemaPropertyFilter::CustomCallback CustomCallback;
+	CustomCallback.BindStatic(CustomJsonSchema);
+	FJsonSchemaPropertyFilter PropertyFilter;
+	PropertyFilter.CustomCb = &CustomCallback;
+	return FJsonSchemaGenerator::FPropertyToJsonSchemaObject(Property, PropertyFilter);
+}
+
+TSharedPtr<FJsonValue> PropertyToJsonData(
+	TNotNull<FProperty*> Property, TNotNull<const void*> Value)
+{
+	RequiredModules::EnsureLoaded();
+	FJsonObjectConverter::CustomExportCallback CustomExportCallback;
+	CustomExportCallback.BindStatic(CustomExport);
+	return FJsonObjectConverter::UPropertyToJsonValue(Property, Value, 0, 0, &CustomExportCallback);
+}
+
+bool JsonDataToProperty(
+	const TSharedPtr<FJsonValue>& JsonValue, TNotNull<FProperty*> Property, TNotNull<void*> OutValue)
+{
+	RequiredModules::EnsureLoaded();
+	FJsonObjectConverter::CustomImportCallback CustomImportCallback;
+	CustomImportCallback.BindStatic(CustomImport);
+	return FJsonObjectConverter::JsonValueToUProperty(
+		JsonValue, Property, OutValue, 0, 0, false, nullptr, &CustomImportCallback);
+}
+
+TSharedPtr<FJsonObject> StructToJsonSchema(
+	TNotNull<const UStruct*> StructDefinition, bool UserVisiblePropertiesOnly)
+{
+	RequiredModules::EnsureLoaded();
+	FJsonSchemaPropertyFilter::CustomCallback CustomCallback;
+	CustomCallback.BindStatic(CustomJsonSchema);
+	FJsonSchemaPropertyFilter PropertyFilter;
+	PropertyFilter.CustomCb = &CustomCallback;
+	PropertyFilter.CheckFlags = GetCheckFlags(UserVisiblePropertiesOnly);
+	PropertyFilter.SkipFlags = CPF_Deprecated;
+	return FJsonSchemaGenerator::UStructToJsonSchemaObject(StructDefinition, PropertyFilter);
+}
+
+TSharedPtr<FJsonObject> StructToJsonData(
+	TNotNull<const UStruct*> StructDefinition, TNotNull<const void*> Struct, bool UserVisiblePropertiesOnly)
+{
+	RequiredModules::EnsureLoaded();
+	FJsonObjectConverter::CustomExportCallback CustomExportCallback;
+	CustomExportCallback.BindStatic(CustomExport);
+	TSharedPtr<FJsonObject> JsonObject = MakeShared<FJsonObject>();
+	EPropertyFlags CheckFlags = GetCheckFlags(UserVisiblePropertiesOnly);
+	EPropertyFlags SkipFlags = CPF_Deprecated;
+	if (FJsonObjectConverter::UStructToJsonObject(
+		StructDefinition, Struct, JsonObject.ToSharedRef(), CheckFlags, SkipFlags, &CustomExportCallback))
+	{
+		return JsonObject;
+	}
+	return nullptr;
+}
+
+bool JsonDataToStruct(
+	const TSharedRef<FJsonObject>& JsonObject, TNotNull<const UStruct*> StructDefinition, TNotNull<void*> OutStruct)
+{
+	RequiredModules::EnsureLoaded();
+	FJsonObjectConverter::CustomImportCallback CustomImportCallback;
+	CustomImportCallback.BindStatic(CustomImport);
+	return FJsonObjectConverter::JsonObjectToUStruct(
+		JsonObject, StructDefinition, OutStruct, 0, 0, false, nullptr, &CustomImportCallback);
+}
+
+}  // namespace UE::ToolsetRegistry::Internal::ToolsetJson
diff --git a/Engine/Plugins/Experimental/ToolsetRegistry/Source/ToolsetRegistry/Private/ToolsetRegistry/ToolsetJson.h b/Engine/Plugins/Experimental/ToolsetRegistry/Source/ToolsetRegistry/Private/ToolsetRegistry/ToolsetJson.h
new file mode 100644
--- /dev/null
+++ b/Engine/Plugins/Experimental/ToolsetRegistry/Source/ToolsetRegistry/Private/ToolsetRegistry/ToolsetJson.h
@@ -0,0 +1,82 @@
+// Copyright Epic Games, Inc. All Rights Reserved. 
+
+#pragma once
+
+#include "CoreMinimal.h"
+#include "Templates/SharedPointer.h"
+
+#include "ToolsetJson.generated.h"
+
+class FJsonObject;
+class FJsonValue;
+
+/// Represents a reference to a UObject or UClass.
+USTRUCT(BlueprintType)
+struct FToolsetReference
+{
+	GENERATED_BODY()
+public:
+	/// The reference stored as a soft path string.
+	UPROPERTY(BlueprintReadWrite, Category="Toolset")
+	FString RefPath;
+};
+
+/// Contains ToolsetRegistry specific wrappers for JSON schema and data generation.
+/// These wrappers perform custom transformations on specific property types, like
+/// converting UObject* and UClass* to FToolsetReference (vs the default in
+/// JsonUtilities which is a string).
+namespace UE::ToolsetRegistry::Internal::ToolsetJson
+{
+	/*
+	 * Converts an FProperty to JSON schema.
+	 * @param Property The FProperty to convert.
+	 * @return JSON schema that represents the FProperty.
+	 */
+	TSharedPtr<FJsonObject> PropertyToJsonSchema(TNotNull<const FProperty*> Property);
+	/*
+	 * Gets an FProperty on an instance to JSON data.
+	 * @param Property The FProperty to convert.
+	 * @param Value The object or struct to get the value from. 
+	 * @return JSON data that represents the value of the FProperty.
+	 */
+	TSharedPtr<FJsonValue> PropertyToJsonData(
+		TNotNull<FProperty*> Property, TNotNull<const void*> Value);
+	/*
+	 * Set an FProperty on an instance from JSON data.
+	 * @param JsonValue The JSON data to convert.
+	 * @param Property The FProperty to apply the data to.
+	 * @param OutValue The object or struct to modify.
+	 * @return True if the data was converted successfully.
+	 */
+	bool JsonDataToProperty(
+		const TSharedPtr<FJsonValue>& JsonValue, TNotNull<FProperty*> Property, TNotNull<void*> OutValue);
+	
+	/*
+	 * Converts a UStruct to JSON schema.
+	 * @param StructDefinition The UStruct type to convert.
+	 * @param UserVisiblePropertiesOnly Whether to export all properties or only the ones that appear in editor.
+	 * @return JSON schema that represents the UStruct.
+	 */
+	TSharedPtr<FJsonObject> StructToJsonSchema(
+		TNotNull<const UStruct*> StructDefinition, bool UserVisiblePropertiesOnly = false);
+	/*
+	 * Converts a UStruct instance to JSON data.
+	 * @param StructDefinition The UStruct type to convert.
+	 * @param Struct The UStruct instance to convert.
+	 * @param UserVisiblePropertiesOnly Whether to export all properties or only the ones that appear in editor.
+	 * @return JSON data that represents the UStruct instance.
+	 */
+	TSharedPtr<FJsonObject> StructToJsonData(
+		TNotNull<const UStruct*> StructDefinition, TNotNull<const void*> Struct,
+		bool UserVisiblePropertiesOnly = false);
+	/*
+	 * Updates a UStruct instance from JSON data.
+	 * @param JsonObject The JSON data to convert.
+	 * @param StructDefinition The UStruct type.
+	 * @param OutStruct The UStruct instance to modify.
+	 * @return True if the data was applied successfully.
+	 */
+	bool JsonDataToStruct(
+		const TSharedRef<FJsonObject>& JsonObject, TNotNull<const UStruct*> StructDefinition,
+		TNotNull<void*> OutStruct);
+}  // namespace UE::ToolsetRegistry::Internal::ToolsetJson
diff --git a/Engine/Plugins/Experimental/ToolsetRegistry/Source/ToolsetRegistry/Private/ToolsetRegistryModule.cpp b/Engine/Plugins/Experimental/ToolsetRegistry/Source/ToolsetRegistry/Private/ToolsetRegistryModule.cpp
--- a/Engine/Plugins/Experimental/ToolsetRegistry/Source/ToolsetRegistry/Private/ToolsetRegistryModule.cpp
+++ b/Engine/Plugins/Experimental/ToolsetRegistry/Source/ToolsetRegistry/Private/ToolsetRegistryModule.cpp
@@ -2,7 +2,6 @@
 
 
 #include "ToolsetRegistryModule.h"
-#include "ToolsetRegistryToolCall.h"
 #include "Modules/ModuleManager.h"
 
 #define LOCTEXT_NAMESPACE "FToolsetRegistryModule"
@@ -15,7 +14,6 @@
 
 /*virtual*/ void FToolsetRegistryModule::StartupModule()
 {
-	UE::ToolsetRegistry::FObjectFunctionToolCall::Initialize();
 }
 
 
diff --git a/Engine/Plugins/Experimental/ToolsetRegistry/Source/ToolsetRegistry/Private/ToolsetRegistryToolCall.cpp b/Engine/Plugins/Experimental/ToolsetRegistry/Source/ToolsetRegistry/Private/ToolsetRegistryToolCall.cpp
deleted file mode 100644
diff --git a/Engine/Plugins/Experimental/ToolsetRegistry/Source/ToolsetRegistry/Private/ToolsetRegistryToolset.cpp b/Engine/Plugins/Experimental/ToolsetRegistry/Source/ToolsetRegistry/Private/ToolsetRegistryToolset.cpp
deleted file mode 100644
diff --git a/Engine/Plugins/Experimental/ToolsetRegistry/Source/ToolsetRegistry/Public/ToolsetRegistry/BlueprintLibraryToolset.h b/Engine/Plugins/Experimental/ToolsetRegistry/Source/ToolsetRegistry/Public/ToolsetRegistry/BlueprintLibraryToolset.h
new file mode 100644
--- /dev/null
+++ b/Engine/Plugins/Experimental/ToolsetRegistry/Source/ToolsetRegistry/Public/ToolsetRegistry/BlueprintLibraryToolset.h
@@ -0,0 +1,95 @@
+// Copyright Epic Games, Inc. All Rights Reserved.
+
+#pragma once
+
+#include "CoreMinimal.h"
+
+#include "Dom/JsonObject.h"
+#include "Kismet/BlueprintFunctionLibrary.h"
+#include "Templates/SharedPointer.h"
+
+#include "ToolsetRegistry/Toolset.h"
+
+#define UE_API TOOLSETREGISTRY_API
+
+namespace UE::ToolsetRegistry
+{
+	class UE_API FObjectFunctionToolCall;
+
+	// Wraps a Blueprint Function Library for use as a Toolset.
+	class FBlueprintLibraryToolset : public FToolset
+	{
+	  public:
+		UE_API FBlueprintLibraryToolset(UClass* LibraryClass);
+
+		// Executes a tool with the given name and JSON input, returning a future with the JSON
+		// output.
+		UE_API virtual TFuture<TValueOrError<FString, FString>> ExecuteTool(
+			const FString& ToolName, const FString& JsonInput) override;
+
+		// Returns the JSON schema for the toolset
+		UE_API virtual FString GetToolsetJsonSchema() const override;
+
+		// Returns the name of the toolset
+		virtual FString GetToolsetName() const override
+		{
+			return ToolsetName;
+		}
+
+		// Returns the version of the toolset
+		virtual FString GetToolsetVersion() const override
+		{
+			// TODO: Investigate some way for toolset to provide a version string.
+			return FString("1.0");
+		}
+
+		// Returns the description of the toolset from the class's metadata.
+		FString GetToolsetDescription() const
+		{
+			return ToolsetClass.Pin() ? ToolsetClass->GetToolTipText().ToString() : FString();
+		}
+
+		// Invoke a tool in this toolset by name.
+		//
+		// Expected to be called by the toolset registry ExecuteTool method, so
+		// ToolName should be the unqualified function name.
+		//
+		// Returns JSON result in OutResult unless there was an error.
+		//
+		// Returns error message in OutError if there was error in invocation.
+		UE_API void InvokeTool(
+			const FString& ToolName, const FString& JsonInput,
+			FString& OutResult, FString& OutError);
+
+		UE_API TArray<FString> GetToolNames();
+
+		// Whether the toolset has a valid set of tools.
+		UE_API bool HasValidTools();
+
+		// Get qualified name for toolset class.
+		UE_API static FString GetToolsetClassName(UClass* InClass);
+
+	  private:
+
+		// The wrapped toolset function library
+		TWeakObjectPtr<UClass> ToolsetClass;
+
+		// Cached toolset name to allow unregistering even when the wrapped toolset class is no longer
+		// available.
+		FString ToolsetName;
+
+		// Map of valid tool names to tool call objects.
+		TMap<FString, TSharedPtr<FObjectFunctionToolCall>> Tools;
+
+		// Generate schema from reflection data.
+		TSharedPtr<FJsonObject> GenerateSchema() const;
+
+		// Generate Tools list.
+		void GenerateToolCallObjects();
+
+		// Helper method to generate list of valid tools.
+		TArray<TObjectPtr<UFunction>> GetValidToolMethodsList() const;
+	};
+}
+
+#undef UE_API
diff --git a/Engine/Plugins/Experimental/ToolsetRegistry/Source/ToolsetRegistry/Public/ToolsetRegistryToolCall.h b/Engine/Plugins/Experimental/ToolsetRegistry/Source/ToolsetRegistry/Public/ToolsetRegistryToolCall.h
deleted file mode 100644
diff --git a/Engine/Plugins/Experimental/ToolsetRegistry/Source/ToolsetRegistry/Public/ToolsetRegistryToolset.h b/Engine/Plugins/Experimental/ToolsetRegistry/Source/ToolsetRegistry/Public/ToolsetRegistryToolset.h
deleted file mode 100644
