Fix loading of AgentSkills from both subclasses defined in blueprint files as well as sub-classes defined in C++.

#Synced-CL 51002830

diff --git a/Engine/Plugins/Experimental/ToolsetRegistry/Source/ToolsetRegistry/Private/AgentSkill.cpp b/Engine/Plugins/Experimental/ToolsetRegistry/Source/ToolsetRegistry/Private/AgentSkill.cpp
--- a/Engine/Plugins/Experimental/ToolsetRegistry/Source/ToolsetRegistry/Private/AgentSkill.cpp
+++ b/Engine/Plugins/Experimental/ToolsetRegistry/Source/ToolsetRegistry/Private/AgentSkill.cpp
@@ -2,11 +2,13 @@
 
 #include "AgentSkill.h"
 
-#include "AssetRegistry/AssetRegistryModule.h"
 #include "AssetRegistry/AssetData.h"
+#include "AssetRegistry/AssetRegistryHelpers.h"
+#include "AssetRegistry/AssetRegistryModule.h"
 #include "Interfaces/ITargetPlatform.h"
 
 #include "ToolsetRegistry/BlueprintLibraryToolset.h"
+#include "ToolsetRegistryModule.h"
 #include "ToolsetRegistry/Toolset.h"
 
 #include UE_INLINE_GENERATED_CPP_BY_NAME(AgentSkill)
@@ -29,36 +31,51 @@ TMap<FString, FString> UAgentSkillToolset::ListSkills()
 {	
 	FAssetRegistryModule& AssetRegistryModule = FModuleManager::LoadModuleChecked<FAssetRegistryModule>("AssetRegistry");
 	IAssetRegistry& AssetRegistry = AssetRegistryModule.Get();
-		
-	FARFilter Filter;
-	Filter.ClassPaths.Add(UAgentSkill::StaticClass()->GetClassPathName());
-	Filter.bRecursiveClasses = true;
-	TArray<FAssetData> AssetDataList;
-	AssetRegistry.GetAssets(Filter, AssetDataList);
-	
+
+	TSet<FTopLevelAssetPath> DerivedNames;
+	AssetRegistry.GetDerivedClassNames({ UAgentSkill::StaticClass()->GetClassPathName() }, {}, DerivedNames);
+
 	TMap<FString, FString> Summaries;
-	for (const auto& AssetData : AssetDataList)
+	for (const FTopLevelAssetPath& DerivedName : DerivedNames)
 	{
-		UObject* LoadedAsset = AssetData.GetAsset();
-		UAgentSkill* AgentSkill = Cast<UAgentSkill>(LoadedAsset);
-		if (AgentSkill)
+		// We ignore skeleton classes (contained in BPs) but consider all other AgentSkill subclasses.
+		if (!DerivedName.GetAssetName().ToString().StartsWith(TEXT("SKEL_")))
 		{
-			Summaries.Add(AssetData.PackageName.ToString(), AgentSkill->Description);
+			FString PathName = DerivedName.ToString();
+			UClass* LoadedClass = FSoftClassPath(PathName).TryLoadClass<UAgentSkill>();
+			if (LoadedClass)
+			{
+				UAgentSkill* AgentSkill = Cast<UAgentSkill>(LoadedClass->GetDefaultObject());
+				if (AgentSkill && !AgentSkill->Description.IsEmpty())
+				{
+					Summaries.Add(PathName, AgentSkill->Description);
+				}
+			}
+			else
+			{
+				UE_LOG(LogToolsetRegistry, Warning,
+					   TEXT("Unable to load AgentSkill subclass: %s"), *DerivedName.ToString());
+			}
 		}
 	}
+
 	return Summaries;
 }
 
 TMap<FString, FAgentSkillDetails> UAgentSkillToolset::GetSkills(const TArray<FString>& SkillPaths)
-{	
+{
 	TMap<FString, FAgentSkillDetails> Skills;
 	for (const auto& SkillPath : SkillPaths)
 	{
-		UAgentSkill* AgentSkill = LoadObject<UAgentSkill>(nullptr, SkillPath);
-		if (AgentSkill)
-		{			
-			Skills.Add(SkillPath, AgentSkill->GetDetails());
+		UClass* AgentSkillClass = LoadObject<UClass>(nullptr, SkillPath);
+		if (AgentSkillClass)
+		{
+			UAgentSkill* AgentSkill = Cast<UAgentSkill>(AgentSkillClass->GetDefaultObject());
+			if (AgentSkill)
+			{
+				Skills.Add(SkillPath, AgentSkill->GetDetails());
+			}
 		}
 	}	
 	return Skills;
-}
\ No newline at end of file
+}
diff --git a/Engine/Plugins/Experimental/ToolsetRegistry/Source/ToolsetRegistry/Private/Tests/AgentSkillsTest.cpp b/Engine/Plugins/Experimental/ToolsetRegistry/Source/ToolsetRegistry/Private/Tests/AgentSkillsTest.cpp
--- a/Engine/Plugins/Experimental/ToolsetRegistry/Source/ToolsetRegistry/Private/Tests/AgentSkillsTest.cpp
+++ b/Engine/Plugins/Experimental/ToolsetRegistry/Source/ToolsetRegistry/Private/Tests/AgentSkillsTest.cpp
@@ -6,17 +6,19 @@
 #include "Misc/AutomationTest.h"
 
 #include "AgentSkill.h"
+#include "EditorAssetLibrary.h"
 #include "FileHelpers.h"
 #include "HAL/FileManager.h"
-#include "ObjectTools.h"
+#include "Kismet2/KismetEditorUtilities.h"
 #include "Misc/Guid.h"
 #include "Misc/PackageName.h"
 #include "Misc/Paths.h"
+#include "ObjectTools.h"
 #include "UObject/Package.h"
 #include "UObject/SavePackage.h"
 #include "UObject/Object.h"
 #include "UObject/ObjectMacros.h"
-#include "Tests/AgentSkillsTest.h"
+#include "Engine/Engine.h"
 #include "Tests/ToolsetRegistryTestFlags.h"
 #include "ToolsetRegistry/BlueprintLibraryToolset.h"
 
@@ -24,22 +26,27 @@
 #if WITH_DEV_AUTOMATION_TESTS
 
 BEGIN_DEFINE_SPEC(FAgentSkillSpec, "AI.ToolsetRegistry.AgentSkillSpec",
-	EAutomationTestFlags::EditorContext | EAutomationTestFlags::EngineFilter)	
+	EAutomationTestFlags::EditorContext | EAutomationTestFlags::EngineFilter)
 	FString MountPoint;
 	FString PackagePath;
+	UBlueprint* SkillBlueprint;
 	UAgentSkill* AgentSkill;
+	UAgentSkill* PromptSkill;
 END_DEFINE_SPEC(FAgentSkillSpec)
 
 void FAgentSkillSpec::Define()
 {
 	BeforeEach([this]()
 	{
-		MountPoint = "/Automation/";		
+		MountPoint = "/Automation/";
 		FPackageName::RegisterMountPoint(*MountPoint, FPaths::AutomationTransientDir());
 
 		PackagePath = MountPoint + TEXT("TestSkill");
 		UPackage* Package = CreatePackage(*PackagePath);
-		AgentSkill = NewObject<UAgentSkill>(Package, TEXT("TestSkill"), RF_Public | RF_Standalone);
+		SkillBlueprint = FKismetEditorUtilities::CreateBlueprint(
+			UAgentSkill::StaticClass(), Package, TEXT("TestSkill"), EBlueprintType::BPTYPE_Normal);
+
+		AgentSkill = Cast<UAgentSkill>(SkillBlueprint->GeneratedClass->GetDefaultObject());
 		AgentSkill->Description = "Testy test";
 		AgentSkill->Toolsets.Add(UAgentSkillToolset::StaticClass());
 		AgentSkill->Instructions = "Do important testy things.";
@@ -48,33 +55,54 @@ void FAgentSkillSpec::Define()
 		TArray<UPackage*> PackagesToSave;
 		PackagesToSave.Add(Package);
 		FEditorFileUtils::PromptForCheckoutAndSave(PackagesToSave, false, false);
+
+		PromptSkill = Cast<UAgentSkill>(UAgentSkillCustomPrompt::StaticClass()->GetDefaultObject());
+		PromptSkill->Description = FString("CustomPrompt:");
 	});
 
 	AfterEach([this]()
 	{
+		PromptSkill->Description.Empty();
+		PromptSkill = nullptr;
+		AgentSkill = nullptr;
+		GEngine->ForceGarbageCollection(true);
+		UEditorAssetLibrary::DeleteLoadedAssets({ SkillBlueprint });
+		SkillBlueprint = nullptr;
+		PackagePath.Empty();
 		FPackageName::UnRegisterMountPoint(*MountPoint, FPaths::AutomationTransientDir());
 		MountPoint.Empty();
-		PackagePath.Empty();
-		AgentSkill = nullptr;
 	});
 
-
 	It("List skill summaries", [this]()
-	{		
+	{
 		TMap<FString, FString> Skills = UAgentSkillToolset::ListSkills();
-		FString* Description = Skills.Find(PackagePath);
-		TestTrue("Summary", Description != nullptr);
-		if (Description)
+
+		// AgentSkill in a native class
+		{
+			FString* Description = Skills.Find(PromptSkill->GetClass()->GetPathName());
+			TestTrue("Summary", Description != nullptr);
+			if (Description)
+			{
+				TestEqual("Summary", *Description, PromptSkill->Description);
+			}
+		}
+		// AgentSkill in BP class.
 		{
-			TestEqual("Summary", *Description, AgentSkill->Description);
+			FString* Description = Skills.Find(AgentSkill->GetClass()->GetPathName());
+			TestTrue("Summary", Description != nullptr);
+			if (Description)
+			{
+				TestEqual("Summary", *Description, AgentSkill->Description);
+			}
 		}
 	});
 
 	It("Get skill details", [this]()
 	{
+		FString ClassName = AgentSkill->GetClass()->GetPathName();
 		TMap<FString, FAgentSkillDetails> Summaries =
-			UAgentSkillToolset::GetSkills(TArray<FString>({ PackagePath }));
-		FAgentSkillDetails* Details = Summaries.Find(PackagePath);
+			UAgentSkillToolset::GetSkills(TArray<FString>({ ClassName }));
+		FAgentSkillDetails* Details = Summaries.Find(ClassName);
 		TestTrue("Details", Details != nullptr);
 		if (Details)
 		{
diff --git a/Engine/Plugins/Experimental/ToolsetRegistry/Source/ToolsetRegistry/Private/Tests/AgentSkillsTest.h b/Engine/Plugins/Experimental/ToolsetRegistry/Source/ToolsetRegistry/Private/Tests/AgentSkillsTest.h
--- a/Engine/Plugins/Experimental/ToolsetRegistry/Source/ToolsetRegistry/Private/Tests/AgentSkillsTest.h
+++ b/Engine/Plugins/Experimental/ToolsetRegistry/Source/ToolsetRegistry/Private/Tests/AgentSkillsTest.h
@@ -7,7 +7,7 @@
 
 #include "AgentSkillsTest.generated.h"
 
-UCLASS()
+UCLASS(Hidden)
 class UAgentSkillCustomPrompt : public UAgentSkill
 {
 	GENERATED_BODY()
