[AI Assistant] Add test automation command to wait for a future

Adds FUntilFutureCommand extending FUntilCommand, which waits until
a future is complete or it times out. This will be used to wait for tool execution
to complete in tests.

#Synced-CL 51016886

diff --git a/Engine/Plugins/Experimental/ToolsetRegistry/Source/ToolsetRegistry/Private/Tests/UntilFutureCommand.h b/Engine/Plugins/Experimental/ToolsetRegistry/Source/ToolsetRegistry/Private/Tests/UntilFutureCommand.h
new file mode 100644
--- /dev/null
+++ b/Engine/Plugins/Experimental/ToolsetRegistry/Source/ToolsetRegistry/Private/Tests/UntilFutureCommand.h
@@ -0,0 +1,43 @@
+// Copyright Epic Games, Inc. All Rights Reserved.
+
+#include "Async/Future.h"
+#include "Misc/AutomationTest.h"
+#include "Templates/Function.h"
+#include "Templates/UnrealTemplate.h"
+
+#if WITH_DEV_AUTOMATION_TESTS
+
+namespace UE::ToolsetRegistry::Internal
+{
+	// Latent test command that waits until a future is complete or a timeout elapses.
+	template<typename ResultT>
+	class FUntilFutureCommand : public FUntilCommand
+	{
+	public:
+		FUntilFutureCommand(
+			TFuture<ResultT>&& InWaitOnFuture,
+			TFunction<void(ResultT&&)> InFutureCompletion,
+			float InTimeout = 5.0f) :
+			FUntilCommand(
+				[this]() -> bool
+				{
+					bool bIsReady = WaitOnFuture.IsReady();
+					if (bIsReady)
+					{
+						OnFutureCompletion(WaitOnFuture.Consume());
+					}
+					return bIsReady;
+				},
+				InTimeout),
+			WaitOnFuture(MoveTemp(InWaitOnFuture)),
+			OnFutureCompletion(InFutureCompletion)
+		{
+		}
+
+	private:
+		TFuture<ResultT> WaitOnFuture;
+		TFunction<void(ResultT&&)> OnFutureCompletion;
+	};
+}
+
+#endif  // WITH_DEV_AUTOMATION_TESTS
\ No newline at end of file
diff --git a/Engine/Plugins/Experimental/ToolsetRegistry/Source/ToolsetRegistry/Private/Tests/UntilFutureCommandTest.cpp b/Engine/Plugins/Experimental/ToolsetRegistry/Source/ToolsetRegistry/Private/Tests/UntilFutureCommandTest.cpp
new file mode 100644
--- /dev/null
+++ b/Engine/Plugins/Experimental/ToolsetRegistry/Source/ToolsetRegistry/Private/Tests/UntilFutureCommandTest.cpp
@@ -0,0 +1,48 @@
+// Copyright Epic Games, Inc. All Rights Reserved.
+
+#include "Async/Future.h"
+#include "Containers/UnrealString.h"
+#include "Misc/AutomationTest.h"
+#include "Templates/SharedPointer.h"
+#include "Templates/ValueOrError.h"
+
+#include "UntilFutureCommand.h"
+
+#if WITH_DEV_AUTOMATION_TESTS
+
+BEGIN_DEFINE_SPEC(
+	FUnitFutureCommandSpec,
+	"AI.ToolsetRegistry.UntilFutureCommand",
+	EAutomationTestFlags::EditorContext | EAutomationTestFlags::EngineFilter)
+END_DEFINE_SPEC(FUnitFutureCommandSpec)
+
+void FUnitFutureCommandSpec::Define()
+{
+	using namespace UE::ToolsetRegistry::Internal;
+
+	Describe(TEXT("FUntilFutureCommand"), [this]()
+	{
+		It(TEXT("should complete when a future completes"), [this]()
+		{
+			TSharedPtr<FString> ExpectedResult = MakeShared<FString>(TEXT("yo"));
+			auto Promise = MakeShared<TPromise<FString>>();
+
+			AddCommand(
+				new FDelayedFunctionLatentCommand(
+					[Promise, ExpectedResult]() -> void
+					{
+						Promise->SetValue(*ExpectedResult);
+					},
+					0.1f /* wait 100ms */));
+			AddCommand(
+				new FUntilFutureCommand<FString>(
+					Promise->GetFuture(),
+					[this, ExpectedResult](FString&& Result) -> void
+					{
+						TestEqual(TEXT("result"), Result, *ExpectedResult);
+					}));
+		});
+	});
+}
+
+#endif  // WITH_DEV_AUTOMATION_TESTS
\ No newline at end of file
