[AI Assistant] Async result for tool calls.

This adds `UToolCallAsyncResult` which `UFUNCTION`s will need to return to
indicate an asynchronous method. It's the responsibility of the function to
complete the `UToolCallAsyncResult` instance on success or failure.

Developers will need to subclass `UToolCallAsyncResult` to return the appropriate
type for their function. The boilerplate required to do so is relatively minimal and
illustrated by `UToolCallAsyncResultString`.

`UToolCallAsyncResultFutureHandler` will be used to connect the blueprint tool
executor (`FToolsetLibraryWrapper`) to yield a future that is completed when
`UToolCallAsyncResult::OnCompleted` is signalled.

Finally, `UE_TOOLSET_REGISTRY_TOOL_CALL_ASYNC_RESULT_WITH_VALUE_SPEC`
is provided to reduce the effort required to write a test for a new result type.
All `UToolCallAsyncResult` classes with values run all tests defined by
`FToolCallAsyncResultBaseSpec::Define()`.

#Synced-CL 51016886

diff --git a/Engine/Plugins/Experimental/ToolsetRegistry/Source/ToolsetRegistry/Private/Tests/ToolCallAsyncResultStringTest.cpp b/Engine/Plugins/Experimental/ToolsetRegistry/Source/ToolsetRegistry/Private/Tests/ToolCallAsyncResultStringTest.cpp
new file mode 100644
--- /dev/null
+++ b/Engine/Plugins/Experimental/ToolsetRegistry/Source/ToolsetRegistry/Private/Tests/ToolCallAsyncResultStringTest.cpp
@@ -0,0 +1,15 @@
+// Copyright Epic Games, Inc. All Rights Reserved.
+
+#include "ToolsetRegistry/ToolCallAsyncResultString.h"
+#include "Tests/ToolCallAsyncResultTest.h"
+
+#if WITH_DEV_AUTOMATION_TESTS
+
+UE_TOOLSET_REGISTRY_TOOL_CALL_ASYNC_RESULT_WITH_VALUE_SPEC(
+	String, UToolCallAsyncResultString, FString,
+	[]() -> FString { return TEXT("Done"); },
+	[](const FString& ExpectedValue) -> TSharedPtr<FJsonValue> {
+		return MakeShared<FJsonValueString>(ExpectedValue);
+	})
+
+#endif  // WITH_DEV_AUTOMATION_TESTS	
\ No newline at end of file
diff --git a/Engine/Plugins/Experimental/ToolsetRegistry/Source/ToolsetRegistry/Private/Tests/ToolCallAsyncResultTest.cpp b/Engine/Plugins/Experimental/ToolsetRegistry/Source/ToolsetRegistry/Private/Tests/ToolCallAsyncResultTest.cpp
new file mode 100644
--- /dev/null
+++ b/Engine/Plugins/Experimental/ToolsetRegistry/Source/ToolsetRegistry/Private/Tests/ToolCallAsyncResultTest.cpp
@@ -0,0 +1,481 @@
+// Copyright Epic Games, Inc. All Rights Reserved.
+
+#include "Tests/ToolCallAsyncResultTest.h"
+
+#include "Async/Async.h"
+#include "Dom/JsonObject.h"
+#include "Misc/AssertionMacros.h"
+#include "Misc/AutomationTest.h"
+#include "HAL/Event.h"
+#include "Templates/FunctionFwd.h"
+#include "Templates/SharedPointer.h"
+#include "Templates/UnrealTemplate.h"
+#include "UObject/Object.h"
+#include "UObject/ObjectPtr.h"
+
+#include "ToolsetRegistry/JsonConversion.h"
+#include "ToolsetRegistry/JsonValueOrError.h"
+#include "ToolsetRegistry/ToolCallAsyncResult.h"
+
+#if WITH_DEV_AUTOMATION_TESTS
+
+// On destruction or when calling Validate(), validates the completion of a UToolCallAsyncResult.
+class FToolCallAsyncResultCompletionValidator
+{
+public:
+	FToolCallAsyncResultCompletionValidator(
+		FAutomationTestBase& TestCaseIn,
+		TObjectPtr<UToolCallAsyncResult> ResultIn) :
+		TestCase(TestCaseIn),
+		Result(ResultIn),
+		ResultHandler(UToolCallAsyncResultFutureHandler::Create(ResultIn)),
+		ResultHandlerFuture(ResultHandler->GetValueAsJson())
+	{
+	}
+
+	virtual ~FToolCallAsyncResultCompletionValidator()
+	{
+		if (!bValidated) Validate();
+	}
+
+	// Prevent copy.
+	FToolCallAsyncResultCompletionValidator(
+		const FToolCallAsyncResultCompletionValidator&) = delete;
+	FToolCallAsyncResultCompletionValidator& operator=(
+		const FToolCallAsyncResultCompletionValidator&) = delete;
+
+	// Check all expectations.
+	void Validate()
+	{
+		check(Result);
+		if (bShouldComplete)
+		{
+			// Verify a complete result.
+			TestCase.TestTrue(TEXT("Result is complete"), Result->bIsComplete);
+			if (ExpectedError.IsEmpty())
+			{
+				EnsureResultCompletedWithValue(GetExpectedJsonValueString());
+			}
+			else
+			{
+				EnsureResultCompletedWithError(ExpectedError, ExpectedError);
+			}
+		}
+		else
+		{
+			EnsureResultIsNotComplete();
+		}
+		bValidated = true;
+	}
+
+	FToolCallAsyncResultCompletionValidator& ExpectShouldComplete(
+		bool bShouldCompleteIn)
+	{
+		bShouldComplete = bShouldCompleteIn;
+		return *this;
+	}
+	
+	FToolCallAsyncResultCompletionValidator& ExpectError(
+		const FString& Error)
+	{
+		ExpectedError = Error;
+		return *this;
+	}
+
+	// A value must be expected to validate a result returns either a value or no value.
+	FToolCallAsyncResultCompletionValidator& ExpectJsonValue(
+		TSharedPtr<FJsonValue> JsonValue)
+	{
+		ExpectedJsonValue = JsonValue;
+		return *this;
+	}
+
+	TStrongObjectPtr<UToolCallAsyncResultFutureHandler> GetResultHandler() const
+	{
+		return ResultHandler;
+	}
+
+	TFuture<UE::ToolsetRegistry::FJsonValueOrError>&
+	GetResultHandlerFuture()
+	{
+		return ResultHandlerFuture;
+	}
+
+	FAutomationTestBase& GetTestCase() { return TestCase; }
+
+	FString GetExpectedJsonValueString() const
+	{
+		check(ExpectedJsonValue);
+		return UE::ToolsetRegistry::Internal::JsonToString(ExpectedJsonValue.ToSharedRef());
+	}
+
+	void EnsureResultCompletedWithValue(const FString& ExpectedJsonValueString)
+	{
+		// With no error.
+		TestCase.TestEqual(TEXT("Result has no error"), Result->Error, TEXT(""));
+		// NOTE: The JSON value is tested indirectly by converting to a string.
+		if (ExpectedJsonValue)
+		{
+			TestCase.TestTrue(TEXT("Result has JSON value"),
+				Result->GetValueAsJson().IsValid());
+			TestCase.TestEqual(
+				TEXT("Result matches expected JSON string"),
+				Result->GetValueAsJsonString(), ExpectedJsonValueString);
+		}
+
+		// Verify the OnCompleted event fired and was handled by the handler.
+		if (TestCase.TestTrue(
+				TEXT("Handler was signaled"), ResultHandlerFuture.IsReady()) &&
+			TestCase.TestTrue(
+				TEXT("Handler result has value"), ResultHandlerFuture.Get().HasValue()))
+		{
+			TestCase.TestEqual(
+				TEXT("Handler result value"),
+				UE::ToolsetRegistry::Internal::JsonToString(
+					ResultHandlerFuture.Get().GetValue().ToSharedRef()),
+				ExpectedJsonValueString);
+		}
+	}
+
+	void EnsureResultCompletedWithError(
+		const FString& InExpectedError, const FString ExpectedHandlerError)
+	{
+		// That reports an error.
+		TestCase.TestEqual(TEXT("Result has an error"), Result->Error, InExpectedError);
+		if (ExpectedJsonValue)
+		{
+			TestCase.TestFalse(
+				TEXT("Result has no JSON value"), Result->GetValueAsJson().IsValid());
+		}
+		TestCase.TestEqual(
+			TEXT("Result has no JSON string"), Result->GetValueAsJsonString(), TEXT(""));
+
+		// Verify the OnCompleted event fired and was handled by the handler.
+		if (TestCase.TestTrue(TEXT("Handler was signaled"), ResultHandlerFuture.IsReady()) &&
+			TestCase.TestTrue(
+				TEXT("Handler result has error"),
+				ResultHandlerFuture.Get().HasError()))
+		{
+			TestCase.TestEqual(
+				TEXT("Handler result error"),
+				ResultHandlerFuture.Get().GetError(), ExpectedHandlerError);
+		}
+	}
+
+	void EnsureResultIsNotComplete()
+	{
+		// Verify a result that isn't complete.
+		TestCase.TestFalse(TEXT("Result is not complete"), Result->bIsComplete);
+		TestCase.TestEqual(TEXT("Result has no error"), Result->Error, TEXT(""));
+		if (ExpectedJsonValue)
+		{
+			TestCase.TestFalse(TEXT("Result has no JSON value"), Result->GetValueAsJson().IsValid());
+		}
+		TestCase.TestEqual(
+			TEXT("Result has no JSON string"), Result->GetValueAsJsonString(), TEXT(""));
+
+		// Verify the OnCompleted event has not fired.
+		TestCase.TestFalse(
+			TEXT("Result handler future is not complete"),
+			ResultHandlerFuture.IsReady());
+	}
+
+private:
+	FAutomationTestBase& TestCase;
+	TStrongObjectPtr<UToolCallAsyncResult> Result;
+	TStrongObjectPtr<UToolCallAsyncResultFutureHandler> ResultHandler;
+	TFuture<UE::ToolsetRegistry::FJsonValueOrError> ResultHandlerFuture;
+	bool bShouldComplete = true;
+	FString ExpectedError;
+	TSharedPtr<FJsonValue> ExpectedJsonValue;
+	bool bValidated = false;
+};
+
+namespace
+{
+	// Handle to a task started with ExecuteOnThread.
+	template<typename ReturnT>
+	class FExecuteOnThreadHandle
+	{
+	public:
+		FExecuteOnThreadHandle(
+			TSharedPtr<FEventRef> InRunTaskEvent,
+			TFuture<ReturnT>&& InReturnValueFuture) :
+			ReturnValueFuture(MoveTemp(InReturnValueFuture)),
+			RunTaskEvent(InRunTaskEvent)
+		{
+		}
+
+		// Prevent copy.
+		FExecuteOnThreadHandle(const FExecuteOnThreadHandle&) = delete;
+		FExecuteOnThreadHandle& operator=(const FExecuteOnThreadHandle&) = delete;
+
+		// Start the task associated with this handle.
+		void Start() { RunTaskEvent->Get()->Trigger(); }
+
+		// Get the return value of the task.
+		ReturnT ReturnValue() const { return ReturnValueFuture.Get();  }
+
+	private:
+		TFuture<ReturnT> ReturnValueFuture;
+		TSharedPtr<FEventRef> RunTaskEvent;
+	};
+
+	// Start a thread, wait until ExecuteOnThreadHandle::Start() is called and then run
+	// Callable.
+	template<typename CallableType>
+	auto ExecuteOnThread(CallableType&& Callable) ->
+		TSharedPtr<FExecuteOnThreadHandle<decltype(Forward<CallableType>(Callable)())>>
+	{
+		TSharedPtr<FEventRef> RunTaskEvent = MakeShared<FEventRef>();
+		auto ReturnValueFuture = Async(
+			EAsyncExecution::Thread,
+			[RunTaskEvent, Callable = MoveTemp(Callable)]() -> auto
+			{
+				RunTaskEvent->Get()->Wait();
+				return Callable();
+			});
+		return MakeShared<FExecuteOnThreadHandle<decltype(Forward<CallableType>(Callable)())>>(
+			RunTaskEvent, MoveTemp(ReturnValueFuture));
+	}
+
+	// Modify an async result in a latent command and then validate its state.
+	class FModifyAndValidateAsyncResultCommand : public FUntilCommand
+	{
+	public:
+		FModifyAndValidateAsyncResultCommand(
+			TSharedPtr<FToolCallAsyncResultCompletionValidator> InValidator,
+			TFunction<bool()>&& InModifyResult) :
+			FUntilCommand(
+				[this]() -> bool
+				{
+					if (!bHasModifiedResult)
+					{
+						if (!ModifyResult()) return true;
+						bHasModifiedResult = true;
+					}
+
+					// Wait for the async result to complete.
+					if (!Validator->GetResultHandlerFuture().IsReady()) return false;
+
+					Validator->Validate();
+					return true;
+				}),
+			Validator(InValidator),
+			ModifyResult(InModifyResult)
+		{
+		}
+
+		FModifyAndValidateAsyncResultCommand(
+			TSharedPtr<FToolCallAsyncResultCompletionValidator> InValidator,
+			TSharedPtr<FExecuteOnThreadHandle<bool>> ExecuteOnThreadHandle,
+			FString ExecuteOnThreadTaskDescription) :
+			FModifyAndValidateAsyncResultCommand(
+				InValidator,
+				[this, ExecuteOnThreadHandle, ExecuteOnThreadTaskDescription]() -> bool
+				{
+					ExecuteOnThreadHandle->Start();
+					return Validator->GetTestCase().TestTrue(
+						*ExecuteOnThreadTaskDescription, ExecuteOnThreadHandle->ReturnValue());
+				})
+		{
+		}
+
+	private:
+		TSharedPtr<FToolCallAsyncResultCompletionValidator> Validator;
+		TFunction<bool()> ModifyResult;
+		bool bHasModifiedResult = false;
+	};
+
+}  // namespace
+
+void FToolCallAsyncResultBaseSpec::Define()
+{
+	static const FString FakeError(TEXT("Failed"));
+
+	// If we can't create a result, skip all tests.
+	if (!CreateResult) return;
+
+	BeforeEach([this]()
+	{
+		Result = CreateResult();
+	});
+
+	It(TEXT("Should not be complete on construction"), [this]()
+	{
+		FToolCallAsyncResultCompletionValidator Validator(*this, Result);
+		Validator.ExpectShouldComplete(false);
+	});
+
+	It(TEXT("Should not send completion if not complete"), [this]()
+	{
+		FToolCallAsyncResultCompletionValidator Validator(*this, Result);
+		Validator.ExpectShouldComplete(false);
+		TestFalse(TEXT("Not complete"), Result->BroadcastOnCompletedIfComplete());
+	});
+
+	It(TEXT("Should complete when an error is set"), [this]()
+	{
+		FToolCallAsyncResultCompletionValidator Validator(*this, Result);
+		Validator.ExpectError(FakeError);
+		TestTrue(TEXT("Set error returns true"), Result->SetError(FakeError));
+	});
+
+	It(TEXT("Should send completion with an error when already completed"), [this]()
+	{
+		Result->SetError(FakeError);
+		FToolCallAsyncResultCompletionValidator Validator(*this, Result);
+		Validator.ExpectError(FakeError);
+		TestTrue(TEXT("Is complete"), Result->BroadcastOnCompletedIfComplete());
+	});
+
+	It(TEXT("Should fail when an empty error is set"), [this]()
+	{
+		FToolCallAsyncResultCompletionValidator Validator(*this, Result);
+		Validator.ExpectShouldComplete(false);
+		AddExpectedMessage(
+			TEXT("Attempted to complete '.*' with an empty error!"),
+			ELogVerbosity::Type::Error, EAutomationExpectedMessageFlags::Contains, 1, true);
+		TestFalse(TEXT("Set an empty error returns false"), Result->SetError(TEXT("")));
+	});
+
+	It(TEXT("Should fail when an error is set more than once"), [this]()
+	{
+		FToolCallAsyncResultCompletionValidator Validator(*this, Result);
+		Validator.ExpectError(FakeError);
+		Result->SetError(FakeError);
+		AddExpectedMessage(
+			TEXT("Unable to complete '.*' more than once"),
+			ELogVerbosity::Type::Error, EAutomationExpectedMessageFlags::Contains, 1, true);
+		TestFalse(TEXT("Set error twice returns false"), Result->SetError(FakeError));
+	});
+
+	It(TEXT("Should complete when an error is set from another thread"), [this]()
+	{
+		auto Validator = MakeShared<
+			FToolCallAsyncResultCompletionValidator>(*this, Result);
+		Validator->ExpectError(FakeError);
+
+		AddCommand(
+			new FModifyAndValidateAsyncResultCommand(
+				Validator,
+				ExecuteOnThread([this]() -> bool { return Result->SetError(FakeError); }),
+				TEXT("Scheduled set error")));
+	});
+
+	It(TEXT("Should send completion with an error on broadcast from another thread"), [this]()
+	{
+		Result->SetError(FakeError);
+		auto Validator = MakeShared<FToolCallAsyncResultCompletionValidator>(*this, Result);
+		Validator->ExpectError(FakeError);
+
+		AddCommand(
+			new FModifyAndValidateAsyncResultCommand(
+				Validator,
+				ExecuteOnThread(
+					[this]() -> bool
+					{
+						return Result->BroadcastOnCompletedIfComplete();
+					}),
+				TEXT("Schedule broadcast on completed")));
+	});
+
+	It(TEXT("Can unsubscribe from a result"), [this]()
+	{
+		FToolCallAsyncResultCompletionValidator Validator(*this, Result);
+		Validator.ExpectShouldComplete(false);
+		Validator.Validate();
+		Validator.GetResultHandler()->Unsubscribe();
+		Result->SetError(FakeError);  // Should have no effect on the unsubscribed handler.
+
+		Validator.EnsureResultCompletedWithError(
+			FakeError, UToolCallAsyncResultFutureHandler::CanceledError);
+	});
+
+	// If a value setter and expected value are not set, skip value tests.
+	if (!(SetResultValue && ExpectedJsonValue)) return;
+
+	It(TEXT("Should complete when a value is set"), [this]()
+	{
+		FToolCallAsyncResultCompletionValidator Validator(*this, Result);
+		Validator.ExpectJsonValue(ExpectedJsonValue);
+		TestTrue(TEXT("Successfully set value"), SetResultValue());
+	});
+
+	It(TEXT("Should send completion with a value when already completed"), [this]()
+	{
+		SetResultValue();
+		FToolCallAsyncResultCompletionValidator Validator(*this, Result);
+		Validator.ExpectJsonValue(ExpectedJsonValue);
+		Result->BroadcastOnCompletedIfComplete();
+	});
+
+	It(TEXT("Should fail when a value is set more than once"), [this]()
+	{
+		FToolCallAsyncResultCompletionValidator Validator(*this, Result);
+		Validator.ExpectJsonValue(ExpectedJsonValue);
+		SetResultValue();
+		AddExpectedMessage(
+			TEXT("Unable to complete '.*' more than once"),
+			ELogVerbosity::Type::Error, EAutomationExpectedMessageFlags::Contains, 1, true);
+		(void)TestFalse(TEXT("Setting the value twice should fail"), SetResultValue());
+	});
+
+	It(TEXT("Should complete when a value is set from another thread"), [this]()
+	{
+		auto Validator = MakeShared<
+			FToolCallAsyncResultCompletionValidator>(*this, Result);
+		Validator->ExpectJsonValue(ExpectedJsonValue);
+
+		AddCommand(
+			new FModifyAndValidateAsyncResultCommand(
+				Validator,
+				ExecuteOnThread([this]() -> bool { return SetResultValue(); }),
+				TEXT("Schedule set value")));
+	});
+
+	It(TEXT("Should send completion with a value on broadcast from another thread"), [this]()
+	{
+		SetResultValue();
+		auto Validator = MakeShared<
+			FToolCallAsyncResultCompletionValidator>(*this, Result);
+		Validator->ExpectJsonValue(ExpectedJsonValue);
+
+		AddCommand(
+			new FModifyAndValidateAsyncResultCommand(
+				Validator,
+				ExecuteOnThread(
+					[this]() -> bool
+					{
+						return Result->BroadcastOnCompletedIfComplete();
+					}),
+				TEXT("Schedule broadcast on completed")));
+	});
+}
+
+// Test functionality of UToolCallAsyncResult where there is no way to successfully complete.
+class FToolCallAsyncResultNoCompletionSpec :
+	public FToolCallAsyncResultBaseSpec
+{
+public:
+	using BaseSpec::BaseSpec;
+
+	virtual void Define() override
+	{
+		SetupCreateResult<UToolCallAsyncResult>();
+		BaseSpec::Define();
+	}
+
+	virtual FString GetBeautifiedTestName() const override
+	{
+		return TEXT("AI.ToolsetRegistry.ToolCallAsyncResult.NoCompletion");
+	}
+};
+
+namespace
+{
+	FToolCallAsyncResultNoCompletionSpec FToolCallAsyncResultNoCompletionSpecInstance(
+		TEXT("FToolCallAsyncResultNoCompletionSpec"));
+}
+
+#endif  // WITH_DEV_AUTOMATION_TESTS
\ No newline at end of file
diff --git a/Engine/Plugins/Experimental/ToolsetRegistry/Source/ToolsetRegistry/Private/Tests/ToolCallAsyncResultTest.h b/Engine/Plugins/Experimental/ToolsetRegistry/Source/ToolsetRegistry/Private/Tests/ToolCallAsyncResultTest.h
new file mode 100644
--- /dev/null
+++ b/Engine/Plugins/Experimental/ToolsetRegistry/Source/ToolsetRegistry/Private/Tests/ToolCallAsyncResultTest.h
@@ -0,0 +1,106 @@
+// Copyright Epic Games, Inc. All Rights Reserved.
+
+#pragma once
+
+#include "Containers/Array.h"
+#include "Containers/UnrealString.h"
+#include "Misc/AutomationTest.h"
+#include "UObject/Object.h"
+#include "UObject/ObjectMacros.h"
+#include "UObject/ObjectPtr.h"
+
+#include "ToolsetRegistry/ToolCallAsyncResult.h"
+#include "ToolsetRegistry/ToolCallAsyncResultFutureHandler.h"
+
+BEGIN_DEFINE_SPEC(
+	FToolCallAsyncResultBaseSpec,
+	"AI.ToolsetRegistry.ToolCallAsyncResult.Base",
+	EAutomationTestFlags::EditorContext | EAutomationTestFlags::EngineFilter)
+	using BaseSpec = FToolCallAsyncResultBaseSpec;
+
+	// Called by test cases to set the result's value.
+	// This *must* be set to execute the tests.
+	TFunction<bool()> SetResultValue;
+	// Queried by test cases to check the result's value.
+	// This *must* be set to execute the tests.
+	TSharedPtr<FJsonValue> ExpectedJsonValue;
+	// Called by before each test case to create the result.
+	TFunction<TObjectPtr<UToolCallAsyncResult>()> CreateResult;
+	// Result operated on in each test.
+	TObjectPtr<UToolCallAsyncResult> Result;
+
+	// Set the CreateResult function to create an AsyncResultT instance.
+	template<typename AsyncResultT>
+	void SetupCreateResult()
+	{
+		CreateResult = [this]()
+			{
+				return TObjectPtr<UToolCallAsyncResult>(NewObject<AsyncResultT>());
+			};
+	}
+
+	// Setup SetResultValue, ExpectedJsonValue and CreateResult.
+	template<typename AsyncResultT, typename ValueT>
+	void Setup(
+		TFunction<TSharedPtr<FJsonValue>(const ValueT&)>& CreateExpectedJsonValue,
+		TFunction<ValueT()>& InCreateExpectedValue)
+	{
+		SetupCreateResult<AsyncResultT>();
+
+		SetResultValue =
+			[this, CreateExpectedValue = InCreateExpectedValue]()
+			{
+				return CastChecked<AsyncResultT>(Result)->SetValue(CreateExpectedValue());
+			};
+
+		ExpectedJsonValue = CreateExpectedJsonValue(InCreateExpectedValue());
+	}
+
+	// NOTE: This intentionally doesn't use END_DEFINE_SPEC() as we want to avoid instantiating
+	// this class.
+};
+
+// Create a test case for a async result with a value type that can be initialized on construction
+// from a constant.
+// 
+// @param AsyncResultValueShortName is the short name of the test case, this is typically the type
+//   under test.
+// @param AsyncResultType is the type of the UToolCallAsyncResult subclass being tested.
+// @param ValueType is type the value of the value that is set to complete the result.
+// @param ExpectedValueFactory is a callable that should create the expected value.
+// @param ExpectedJsonValueFactory is a callable that takes a reference to the expected value
+//   and return TSharedPtr<FJsonValue> to verify the result's value is correctly converted to the
+//   expected JSON value.
+#define UE_TOOLSET_REGISTRY_TOOL_CALL_ASYNC_RESULT_WITH_VALUE_SPEC(                                \
+		AsyncResultValueShortName, AsyncResultType, ValueType, ExpectedValueFactory,               \
+		ExpectedJsonValueFactory)                                                                  \
+	class FToolCallAsyncResult##AsyncResultValueShortName##Spec :                                  \
+		public FToolCallAsyncResultBaseSpec                                                        \
+	{                                                                                              \
+	public:                                                                                        \
+		using BaseSpec::BaseSpec;                                                                  \
+	                                                                                               \
+		virtual void Define() override                                                             \
+		{                                                                                          \
+			Setup<AsyncResultType, ValueType>(JsonValueFactory, ValueFactory);                     \
+			BaseSpec::Define();                                                                    \
+		}                                                                                          \
+	                                                                                               \
+		virtual FString GetBeautifiedTestName() const override                                     \
+		{                                                                                          \
+			return TEXT("AI.ToolsetRegistry.ToolCallAsyncResult."                                  \
+			            #AsyncResultValueShortName);                                               \
+		}                                                                                          \
+		                                                                                           \
+		TFunction<ValueType()> ValueFactory = ExpectedValueFactory;                                \
+		TFunction<TSharedPtr<FJsonValue>(const ValueType&)> JsonValueFactory =                     \
+			ExpectedJsonValueFactory;                                                              \
+	};                                                                                             \
+	                                                                                               \
+	namespace                                                                                      \
+	{                                                                                              \
+		FToolCallAsyncResult##AsyncResultValueShortName##Spec                                      \
+		    FToolCallAsyncResult##AsyncResultValueShortName##SpecInstance(                         \
+		        TEXT(UE_INLINE_STRINGIFY(                                                          \
+		            FToolCallAsyncResultWith##AsyncResultValueShortName##Spec)));                  \
+	}
\ No newline at end of file
diff --git a/Engine/Plugins/Experimental/ToolsetRegistry/Source/ToolsetRegistry/Private/Tests/ToolCallAsyncResultVoidTest.cpp b/Engine/Plugins/Experimental/ToolsetRegistry/Source/ToolsetRegistry/Private/Tests/ToolCallAsyncResultVoidTest.cpp
new file mode 100644
--- /dev/null
+++ b/Engine/Plugins/Experimental/ToolsetRegistry/Source/ToolsetRegistry/Private/Tests/ToolCallAsyncResultVoidTest.cpp
@@ -0,0 +1,44 @@
+// Copyright Epic Games, Inc. All Rights Reserved.
+
+#include "Dom/JsonObject.h"
+#include "Misc/AutomationTest.h"
+#include "Templates/SharedPointer.h"
+#include "UObject/Object.h"
+#include "UObject/ObjectPtr.h"
+
+#include "ToolsetRegistry/ToolCallAsyncResultVoid.h"
+#include "Tests/ToolCallAsyncResultTest.h"
+
+#if WITH_DEV_AUTOMATION_TESTS
+
+class FToolCallAsyncResultVoidSpec :
+	public FToolCallAsyncResultBaseSpec
+{
+public:
+	using BaseSpec::BaseSpec;
+
+	virtual void Define() override
+	{
+		SetupCreateResult<UToolCallAsyncResultVoid>();
+		SetResultValue = [this]()
+			{
+				return CastChecked<UToolCallAsyncResultVoid>(Result)->SetCompleted();
+			};
+		ExpectedJsonValue = MakeShared<FJsonValueNull>();
+		BaseSpec::Define();
+	}
+
+	virtual FString GetBeautifiedTestName() const override
+	{
+		return TEXT("AI.ToolsetRegistry.ToolCallAsyncResult.Void");
+	}
+};
+
+namespace
+{
+	FToolCallAsyncResultVoidSpec
+		FToolCallAsyncResultVoidSpecInstance(
+			TEXT("FToolCallAsyncResultVoidSpec"));
+}
+
+#endif  // WITH_DEV_AUTOMATION_TESTS
\ No newline at end of file
diff --git a/Engine/Plugins/Experimental/ToolsetRegistry/Source/ToolsetRegistry/Private/ToolsetRegistry/ToolCallAsyncResult.cpp b/Engine/Plugins/Experimental/ToolsetRegistry/Source/ToolsetRegistry/Private/ToolsetRegistry/ToolCallAsyncResult.cpp
new file mode 100644
--- /dev/null
+++ b/Engine/Plugins/Experimental/ToolsetRegistry/Source/ToolsetRegistry/Private/ToolsetRegistry/ToolCallAsyncResult.cpp
@@ -0,0 +1,133 @@
+// Copyright Epic Games, Inc. All Rights Reserved.
+
+#include "ToolsetRegistry/ToolCallAsyncResult.h"
+
+#include "Async/Async.h"
+#include "Async/UniqueLock.h"
+#include "Containers/UnrealString.h"
+#include "CoreGlobals.h"
+#include "Logging/LogMacros.h"
+#include "Misc/AssertionMacros.h"
+#include "Templates/Function.h"
+#include "UObject/Class.h"
+
+#include "ToolsetRegistryModule.h"
+#include "ToolsetRegistry/JsonConversion.h"
+#include "ToolsetRegistry/ToolsetJson.h"
+
+TSharedPtr<FJsonValue> UToolCallAsyncResult::GetValueAsJson() const
+{
+	if (!bIsComplete || !Error.IsEmpty()) return TSharedPtr<FJsonValue>();
+
+	UClass* Class = GetClass();
+	check(Class);
+	FProperty* Property = Class->FindPropertyByName(TEXT("Value"));
+	if (!Property)
+	{
+		UE_LOG(
+			LogToolsetRegistry, Error, TEXT("Unable to get Value property from class '%s'"),
+			*Class->GetName());
+		return TSharedPtr<FJsonValue>();
+	}
+	// The non-templated FProperty::ContainerVoidPtrToValuePtrInternal() and
+	// FProperty::PointerToValuePtr() are not exposed, so manually calculate the offset of
+	// the property from the container (UObject) base pointer using the same approach as
+	// FProperty::PointerToValuePtr() here.
+	const uint8* PropertyValuePointer =
+		reinterpret_cast<const uint8*>(this) + Property->GetOffset_ForInternal();
+	return UE::ToolsetRegistry::Internal::ToolsetJson::PropertyToJsonData(
+		Property, PropertyValuePointer);
+}
+
+FString UToolCallAsyncResult::GetValueAsJsonString() const
+{
+	if (Error.IsEmpty())
+	{
+		TSharedPtr<FJsonValue> JsonValue = GetValueAsJson();
+		if (JsonValue)
+		{
+			return UE::ToolsetRegistry::Internal::JsonToString(JsonValue.ToSharedRef());
+		}
+	}
+	return TEXT("");
+}
+
+bool UToolCallAsyncResult::SetError(const FString& InError)
+{
+	if (InError.Len() == 0)
+	{
+		UE_LOG(
+			LogToolsetRegistry, Error,
+			TEXT("Attempted to complete '%s' with an empty error!"),
+			*GetClass()->GetName());
+		return false;
+	}
+	return MaybeBroadcastCompletion(
+		[this, CopiedError = InError]() mutable -> void
+		{
+			this->Error = MoveTemp(CopiedError);
+		});
+}
+
+bool UToolCallAsyncResult::BroadcastOnCompletedIfComplete()
+{
+	bool bAlreadyComplete = GetIsComplete();
+	if (bAlreadyComplete)
+	{
+		TStrongObjectPtr<UToolCallAsyncResult> Result(this);
+		RunOnMainThread([Result = MoveTemp(Result)]() { Result->BroadcastOnCompleted(); });
+	}
+	return bAlreadyComplete;
+}
+
+bool UToolCallAsyncResult::MaybeBroadcastCompletion(
+	TFunction<void()>&& SetFinalState)
+{
+	if (GetIsComplete())
+	{
+		UE_LOG(
+			LogToolsetRegistry, Error, TEXT("Unable to complete '%s' more than once"),
+			*GetClass()->GetName());
+		return false;
+	}
+	TStrongObjectPtr<UToolCallAsyncResult> Result(this);
+	RunOnMainThread(
+		[Result = MoveTemp(Result), SetFinalState = MoveTemp(SetFinalState)]()
+		{
+			SetFinalState();
+			Result->BroadcastOnCompleted();
+		});
+	return true;
+}
+
+void UToolCallAsyncResult::SetIsComplete()
+{
+	UE::TUniqueLock Lock(CompletedLock);
+	bIsComplete = true;
+}
+
+bool UToolCallAsyncResult::GetIsComplete() const
+{
+	UE::TUniqueLock Lock(CompletedLock);
+	return bIsComplete;
+}
+
+void UToolCallAsyncResult::BroadcastOnCompleted()
+{
+	ensureAlways(IsInGameThread());
+	SetIsComplete();
+	OnCompleted.Broadcast(this);
+	OnCompleted.Clear();
+}
+
+void UToolCallAsyncResult::RunOnMainThread(TUniqueFunction<void()>&& Function)
+{
+	if (IsInGameThread())
+	{
+		Function();
+	}
+	else
+	{
+		AsyncTask(ENamedThreads::GameThread, MoveTemp(Function));
+	}
+}
\ No newline at end of file
diff --git a/Engine/Plugins/Experimental/ToolsetRegistry/Source/ToolsetRegistry/Private/ToolsetRegistry/ToolCallAsyncResultFutureHandler.cpp b/Engine/Plugins/Experimental/ToolsetRegistry/Source/ToolsetRegistry/Private/ToolsetRegistry/ToolCallAsyncResultFutureHandler.cpp
new file mode 100644
--- /dev/null
+++ b/Engine/Plugins/Experimental/ToolsetRegistry/Source/ToolsetRegistry/Private/ToolsetRegistry/ToolCallAsyncResultFutureHandler.cpp
@@ -0,0 +1,90 @@
+// Copyright Epic Games, Inc. All Rights Reserved.
+
+#include "ToolsetRegistry/ToolCallAsyncResultFutureHandler.h"
+
+#include "CoreGlobals.h"
+#include "Misc/AssertionMacros.h"
+
+#include "ToolsetRegistry/JsonConversion.h"
+
+const FString UToolCallAsyncResultFutureHandler::CanceledError(TEXT("Canceled"));
+
+UToolCallAsyncResultFutureHandler::~UToolCallAsyncResultFutureHandler()
+{
+	Unsubscribe(true);
+}
+
+TFuture<UE::ToolsetRegistry::FJsonValueOrError>
+UToolCallAsyncResultFutureHandler::GetValueAsJson()
+{
+	check(!bRetrievedFuture);
+	bRetrievedFuture = true;
+	// Hold a strong reference to this object within the future so that it can't be garbage
+	// collected until the future is complete.
+	TStrongObjectPtr<UToolCallAsyncResultFutureHandler> This(this);
+	return ResultPromise.GetFuture().Next(
+		[This](UE::ToolsetRegistry::FJsonValueOrError&& Result) ->
+			UE::ToolsetRegistry::FJsonValueOrError
+		{
+			return Result;
+		});
+}
+
+void UToolCallAsyncResultFutureHandler::Subscribe(UToolCallAsyncResult* ResultIn)
+{
+	check(IsInGameThread());
+	check(ResultIn);
+	check(!Result);
+	Result = ResultIn;
+	Result->OnCompleted.AddUniqueDynamic(this, &ThisClass::OnCompleted);
+}
+
+void UToolCallAsyncResultFutureHandler::Unsubscribe()
+{
+	Unsubscribe(false);
+}
+
+void UToolCallAsyncResultFutureHandler::Unsubscribe(bool bIsDestructing)
+{
+	check(IsInGameThread());
+	// Cancel the promise if it was never fulfilled.
+	if (!bOnCompletedHandled)
+	{
+		ResultPromise.SetValue(MakeError(CanceledError));
+		if (!bIsDestructing && Result)
+		{
+			Result->OnCompleted.RemoveDynamic(this, &ThisClass::OnCompleted);
+		}
+	}
+	bOnCompletedHandled = true;
+	Result = nullptr;
+}
+
+void UToolCallAsyncResultFutureHandler::OnCompleted(UToolCallAsyncResult* ResultIn)
+{
+	check(IsInGameThread());
+	check(Result == ResultIn);
+	check(Result->bIsComplete);
+	check(!bOnCompletedHandled);
+	const FString& Error = Result->Error;
+	if (Error.IsEmpty())
+	{
+		ResultPromise.SetValue(MakeValue(Result->GetValueAsJson()));
+	}
+	else
+	{
+		ResultPromise.SetValue(MakeError(Error));
+	}
+	bOnCompletedHandled = true;
+	Unsubscribe();
+}
+
+TStrongObjectPtr<UToolCallAsyncResultFutureHandler>
+UToolCallAsyncResultFutureHandler::Create(TObjectPtr<UToolCallAsyncResult> Result)
+{
+	UToolCallAsyncResultFutureHandler* OnCompletedHandler =
+		NewObject<UToolCallAsyncResultFutureHandler>();
+	OnCompletedHandler->Subscribe(Result);
+	return TStrongObjectPtr<UToolCallAsyncResultFutureHandler>(OnCompletedHandler);
+}
+
diff --git a/Engine/Plugins/Experimental/ToolsetRegistry/Source/ToolsetRegistry/Private/ToolsetRegistry/ToolCallAsyncResultFutureHandler.h b/Engine/Plugins/Experimental/ToolsetRegistry/Source/ToolsetRegistry/Private/ToolsetRegistry/ToolCallAsyncResultFutureHandler.h
new file mode 100644
--- /dev/null
+++ b/Engine/Plugins/Experimental/ToolsetRegistry/Source/ToolsetRegistry/Private/ToolsetRegistry/ToolCallAsyncResultFutureHandler.h
@@ -0,0 +1,63 @@
+// Copyright Epic Games, Inc. All Rights Reserved.
+
+#pragma once
+
+#include "Async/Future.h"
+#include "Containers/UnrealString.h"
+#include "Delegates/Delegate.h"
+#include "Dom/JsonValue.h"
+#include "UObject/Object.h"
+#include "UObject/ObjectMacros.h"
+#include "UObject/ObjectPtr.h"
+
+#include "ToolsetRegistry/JsonValueOrError.h"
+#include "ToolsetRegistry/ObjectFunctionToolCall.h"
+#include "ToolsetRegistry/ToolCallAsyncResult.h"
+
+#include "ToolCallAsyncResultFutureHandler.generated.h"
+
+
+// Handles UToolCallAsyncResult::OnCompleted by completing the associated future.
+//
+// NOTE: This is tested by ToolCallAsyncResultTest.cpp.
+UCLASS(MinimalAPI, NotBlueprintable, HideDropdown)
+class UToolCallAsyncResultFutureHandler : public UObject
+{
+	GENERATED_BODY()
+
+public:
+	virtual ~UToolCallAsyncResultFutureHandler() override;
+
+	// Resolves with the JSON result of UToolCallAsyncResult object passed to Create().
+	// NOTE: The future can only be retrieved once.
+	TFuture<UE::ToolsetRegistry::FJsonValueOrError> GetValueAsJson();
+
+	// Unsubscribe from the result and cancel the pending future returned from GetValueAsJson().
+	void Unsubscribe();
+
+private:
+	// Handles a result completion.
+	UFUNCTION()
+	void OnCompleted(UToolCallAsyncResult* Result);
+
+	// Subscribes to result completion.
+	void Subscribe(UToolCallAsyncResult* Result);
+
+	// Unsubscribe, optionally without touching the Result as it may already be garbage collected.
+	void Unsubscribe(bool bIsDestructing);
+
+private:
+	TObjectPtr<UToolCallAsyncResult> Result;
+	TPromise<UE::ToolsetRegistry::FJsonValueOrError> ResultPromise;
+	bool bOnCompletedHandled = false;
+	bool bRetrievedFuture = false;
+
+public:
+	// Create a handler and bind the OnCompleted method to the specified result instance's
+	// OnCompleted delegate.
+	static TStrongObjectPtr<UToolCallAsyncResultFutureHandler> Create(
+		TObjectPtr<UToolCallAsyncResult> Result);
+
+	// Error that indicates the future was canceled.
+	static const FString CanceledError;
+};
diff --git a/Engine/Plugins/Experimental/ToolsetRegistry/Source/ToolsetRegistry/Private/ToolsetRegistry/ToolCallAsyncResultVoid.cpp b/Engine/Plugins/Experimental/ToolsetRegistry/Source/ToolsetRegistry/Private/ToolsetRegistry/ToolCallAsyncResultVoid.cpp
new file mode 100644
--- /dev/null
+++ b/Engine/Plugins/Experimental/ToolsetRegistry/Source/ToolsetRegistry/Private/ToolsetRegistry/ToolCallAsyncResultVoid.cpp
@@ -0,0 +1,11 @@
+// Copyright Epic Games, Inc. All Rights Reserved.
+
+#include "ToolsetRegistry/ToolCallAsyncResultVoid.h"
+
+#include "Dom/JsonValue.h"
+#include "Templates/SharedPointer.h"
+
+TSharedPtr<FJsonValue> UToolCallAsyncResultVoid::GetValueAsJson() const
+{
+	return bIsComplete ? MakeShared<FJsonValueNull>() : TSharedPtr<FJsonValue>();
+}
diff --git a/Engine/Plugins/Experimental/ToolsetRegistry/Source/ToolsetRegistry/Public/ToolsetRegistry/ToolCallAsyncResult.h b/Engine/Plugins/Experimental/ToolsetRegistry/Source/ToolsetRegistry/Public/ToolsetRegistry/ToolCallAsyncResult.h
new file mode 100644
--- /dev/null
+++ b/Engine/Plugins/Experimental/ToolsetRegistry/Source/ToolsetRegistry/Public/ToolsetRegistry/ToolCallAsyncResult.h
@@ -0,0 +1,146 @@
+// Copyright Epic Games, Inc. All Rights Reserved.
+
+#pragma once
+
+#include "Async/Mutex.h"
+#include "Containers/UnrealString.h"
+#include "Delegates/DelegateCombinations.h"
+#include "Dom/JsonValue.h"
+#include "Templates/Function.h"
+#include "Templates/SharedPointer.h"
+#include "Templates/UnrealTemplate.h"
+#include "UObject/Object.h"
+#include "UObject/ObjectMacros.h"
+
+#include "ToolCallAsyncResult.generated.h"
+
+#define UE_API TOOLSETREGISTRY_API
+
+/// Notified when an async tool call completes successfully or with an error.
+DECLARE_DYNAMIC_MULTICAST_DELEGATE_OneParam(FToolCallAsyncResultCompleted,
+	UToolCallAsyncResult*, Result);
+
+/// Base class for an asynchronous tool call result that is analogous to a promise.
+///
+/// Asynchronous tools should *not* use this base class directly, instead they should use an 
+/// either define or use an existing derivative of this base class for asynchronous results.
+///
+/// Asynchronous tools return derivatives of this class to indicate a pending result.
+/// When an asynchronous tool completes its operation it should perform one of the following
+/// operations:
+/// * If successful, use MaybeBroadcastSuccessfulCompletion() from the derived class
+///   - ideally implemented in a method called `SetValue` - to notify listeners of the
+///   `OnCompleted` delegate, set bIsComplete to true and set the result's value.
+/// * If an error occurs, use SetError() to set the `Error` property, set bIsComplete to true
+///    and notify listeners of the `OnCompleted` delegate.
+///
+UCLASS(BlueprintType, Blueprintable)
+class UToolCallAsyncResult : public UObject
+{
+	GENERATED_BODY()
+
+public:
+	/// Whether the associated tool call is complete and the result via Value and GetValueAsJson()
+	/// or Error are available to read.
+	///
+	/// @warning Do *not* set this from C++ instead use SetError() or SetValue() if implemented by
+	/// the subclass.
+	UPROPERTY(BlueprintReadOnly, Category = "ToolsetRegistry")
+	bool bIsComplete = false;
+
+	/// Notified when the tool call is completes successfully or with an error.
+	UPROPERTY(BlueprintAssignable, Category = "ToolsetRegistry")
+	FToolCallAsyncResultCompleted OnCompleted;
+
+	/// If this is an non-empty string, an error occurred while executing the tool that returned
+	/// this instance. An empty string indicates either no error occurred or the associated tool
+	/// call is not complete.
+	/// 
+	/// @warning Do *not* set this from C++ instead use SetError() or SetValue() if implemented by
+	/// the subclass.
+	UPROPERTY(BlueprintReadOnly, Category = "ToolsetRegistry")
+	FString Error;
+
+	/// Get the JSON representation of the Value. If the associated tool call is not complete or an
+	/// error occurred, this returns an empty string.
+	UFUNCTION(BlueprintCallable, Category = "ToolsetRegistry")
+	FString GetValueAsJsonString() const;
+
+	/// Get the JSON representation of the Value. If the associated tool call is not complete or an
+	/// error occurred, this returns an empty (null) shared pointer.
+	UE_API virtual TSharedPtr<FJsonValue> GetValueAsJson() const;
+
+	/// Complete this result with an error and notify listeners of OnCompleted.
+	///
+	/// @returns true the result was completed with the error, false otherwise.
+	///
+	/// @note This is thread safe, the Error property will always be updated and OnCompleted
+	/// will be signaled on the main thread.
+	UFUNCTION(BlueprintCallable, Category = "ToolsetRegistry")
+	bool SetError(const FString& InError);
+
+	/// Broadcast to subscribers of the OnCompleted delegate if the result is already complete.
+	///
+	/// This *must* be called after subscribing to OnCompleted to ensure an event receives a
+	/// notification when the result is already complete.
+	///
+	/// @returns true if the result was complete, false otherwise.
+	///
+	/// @note This is thread safe, OnCompleted will always be signaled on the main thread.
+	///
+	/// @remark Ideally we would override subscription methods for the OnCompleted delegate
+	/// and notify subscribers there but unfortunately that isn't possible.
+	UFUNCTION(BlueprintCallable, Category = "ToolsetRegistry")
+	bool BroadcastOnCompletedIfComplete();
+
+protected:
+	/// Broadcast that the result is complete if it hasn't already been completed.
+	///
+	/// @returns true if the broadcast was executed or scheduled, false otherwise.
+	///
+	/// @note To make this thread safe, if this isn't called on the main thread, SetFinalState 
+	/// will not be called until the next main thread update so be careful that SetFinalState does
+	/// not reference any data on the stack.
+	UE_API bool MaybeBroadcastCompletion(TFunction<void()>&& SetFinalState);
+
+	/// Broadcast a successful completion with a value.
+	///
+	/// @returns true if the broadcast was executed or scheduled, false otherwise.
+	///
+	/// @note: InValue must be movable, as ValueToSet must be set on the main thread.
+	template<typename ValueT>
+	bool MaybeBroadcastSuccessfulCompletion(ValueT&& InValue, ValueT& ValueToSet)
+	{
+		return MaybeBroadcastCompletion(
+			[InValue = MoveTemp(InValue), &ValueToSet]() mutable -> void
+			{
+				ValueToSet = MoveTemp(InValue);
+			});
+	}
+
+	/// Broadcast a successful completion with no value.
+	///
+	/// @returns true if the broadcast was executed or scheduled, false otherwise.
+	bool MaybeBroadcastSuccessfulCompletion()
+	{
+		return MaybeBroadcastCompletion([]() -> void {});
+	}
+
+private:
+	// Thread safe access of bIsComplete.
+	void SetIsComplete();
+	bool GetIsComplete() const;
+
+	// Broadcast OnCompleted and clear all listeners from the delegate.
+	void BroadcastOnCompleted();
+
+private:
+	// Guards: bIsComplete.
+	mutable UE::FMutex CompletedLock;
+
+private:
+	// Run a function on the main thread as soon as possible.
+	static void RunOnMainThread(TUniqueFunction<void()>&& Function);
+};
+
+#undef UE_API
\ No newline at end of file
diff --git a/Engine/Plugins/Experimental/ToolsetRegistry/Source/ToolsetRegistry/Public/ToolsetRegistry/ToolCallAsyncResultString.h b/Engine/Plugins/Experimental/ToolsetRegistry/Source/ToolsetRegistry/Public/ToolsetRegistry/ToolCallAsyncResultString.h
new file mode 100644
--- /dev/null
+++ b/Engine/Plugins/Experimental/ToolsetRegistry/Source/ToolsetRegistry/Public/ToolsetRegistry/ToolCallAsyncResultString.h
@@ -0,0 +1,27 @@
+// Copyright Epic Games, Inc. All Rights Reserved.
+
+#pragma once
+
+#include "ToolsetRegistry/ToolCallAsyncResult.h"	
+
+#include "ToolCallAsyncResultString.generated.h"
+
+// Async tool call result that completes with a string.
+UCLASS(BlueprintType)
+class UToolCallAsyncResultString : public UToolCallAsyncResult
+{
+	GENERATED_BODY()
+
+public:
+	// Set the result.
+	UFUNCTION(BlueprintCallable, Category = "ToolsetRegistry")
+	bool SetValue(const FString& InValue)
+	{
+		return MaybeBroadcastSuccessfulCompletion(FString(InValue), Value);
+	}
+
+	// Value of the result if it has been completed. This should only be set via SetValue().
+	UPROPERTY(BlueprintReadOnly, Category = "ToolsetRegistry")
+	FString Value;
+};
+
diff --git a/Engine/Plugins/Experimental/ToolsetRegistry/Source/ToolsetRegistry/Public/ToolsetRegistry/ToolCallAsyncResultVoid.h b/Engine/Plugins/Experimental/ToolsetRegistry/Source/ToolsetRegistry/Public/ToolsetRegistry/ToolCallAsyncResultVoid.h
new file mode 100644
--- /dev/null
+++ b/Engine/Plugins/Experimental/ToolsetRegistry/Source/ToolsetRegistry/Public/ToolsetRegistry/ToolCallAsyncResultVoid.h
@@ -0,0 +1,32 @@
+// Copyright Epic Games, Inc. All Rights Reserved.
+
+#pragma once
+
+#include "Dom/JsonValue.h"
+#include "Templates/SharedPointer.h"
+#include "UObject/ObjectMacros.h"
+
+#include "ToolsetRegistry/ToolCallAsyncResult.h"
+
+#include "ToolCallAsyncResultVoid.generated.h"
+
+#define UE_API TOOLSETREGISTRY_API
+
+// Async tool call result that completes with no value.
+UCLASS(BlueprintType)
+class UToolCallAsyncResultVoid : public UToolCallAsyncResult	
+{
+	GENERATED_BODY()
+
+public:
+	// Complete this result with no value and notify listeners of OnCompleted.
+	UFUNCTION(BlueprintCallable, Category = "ToolsetRegistry")
+	bool SetCompleted()
+	{
+		return MaybeBroadcastSuccessfulCompletion();
+	}
+
+	UE_API virtual TSharedPtr<FJsonValue> GetValueAsJson() const override;
+};
+
+#undef UE_API 
