[AI Assistant] Fix warning in UntilFutureCommand with some compilers.

#rnx

#Synced-CL 51035176

diff --git a/Engine/Plugins/Experimental/ToolsetRegistry/Source/ToolsetRegistry/Private/Tests/UntilFutureCommand.h b/Engine/Plugins/Experimental/ToolsetRegistry/Source/ToolsetRegistry/Private/Tests/UntilFutureCommand.h
--- a/Engine/Plugins/Experimental/ToolsetRegistry/Source/ToolsetRegistry/Private/Tests/UntilFutureCommand.h
+++ b/Engine/Plugins/Experimental/ToolsetRegistry/Source/ToolsetRegistry/Private/Tests/UntilFutureCommand.h
@@ -11,14 +11,16 @@ namespace UE::ToolsetRegistry::Internal
 {
 	// Latent test command that waits until a future is complete or a timeout elapses.
 	template<typename ResultT>
-	class FUntilFutureCommand : public FUntilCommand
+	class FUntilFutureCommand : public IAutomationLatentCommand
 	{
 	public:
 		FUntilFutureCommand(
 			TFuture<ResultT>&& InWaitOnFuture,
 			TFunction<void(ResultT&&)> InFutureCompletion,
 			float InTimeout = 5.0f) :
-			FUntilCommand(
+			WaitOnFuture(MoveTemp(InWaitOnFuture)),
+			OnFutureCompletion(InFutureCompletion),
+			UntilCommand(
 				[this]() -> bool
 				{
 					bool bIsReady = WaitOnFuture.IsReady();
@@ -28,15 +30,19 @@ namespace UE::ToolsetRegistry::Internal
 					}
 					return bIsReady;
 				},
-				InTimeout),
-			WaitOnFuture(MoveTemp(InWaitOnFuture)),
-			OnFutureCompletion(InFutureCompletion)
+				InTimeout)
+		{
+		}
+
+		virtual bool Update() override
 		{
+			return UntilCommand.Update();
 		}
 
 	private:
 		TFuture<ResultT> WaitOnFuture;
 		TFunction<void(ResultT&&)> OnFutureCompletion;
+		FUntilCommand UntilCommand;
 	};
 }
 
