[AI Assistant] Utility to reduce the verbosity of creating fulfilled futures.

TValueOrErrorFuture is a simple class that makes it easier / less verbose to create
fulfilled futures that resolve to a TValueOrError.

#Synced-CL 51035176

diff --git a/Engine/Plugins/Experimental/ToolsetRegistry/Source/ToolsetRegistry/Private/Tests/ValueOrErrorFutureTest.cpp b/Engine/Plugins/Experimental/ToolsetRegistry/Source/ToolsetRegistry/Private/Tests/ValueOrErrorFutureTest.cpp
new file mode 100644
--- /dev/null
+++ b/Engine/Plugins/Experimental/ToolsetRegistry/Source/ToolsetRegistry/Private/Tests/ValueOrErrorFutureTest.cpp
@@ -0,0 +1,53 @@
+// Copyright Epic Games, Inc. All Rights Reserved.
+
+#include "Misc/AutomationTest.h"
+#include "Templates/ValueOrError.h"
+
+#include "ToolsetRegistry/ValueOrErrorFuture.h"
+
+#if WITH_DEV_AUTOMATION_TESTS
+
+BEGIN_DEFINE_SPEC(
+	FToolsetRegistryValueOrErrorFutureSpec,
+	"AI.ToolsetRegistry.ValueOrErrorFuture",
+	EAutomationTestFlags::EditorContext | EAutomationTestFlags::EngineFilter)
+END_DEFINE_SPEC(FToolsetRegistryValueOrErrorFutureSpec)
+
+void FToolsetRegistryValueOrErrorFutureSpec::Define()
+{
+	using namespace UE::ToolsetRegistry::Internal;
+
+	Describe(TEXT("MakeError"), [this]()
+	{
+		It(TEXT("should return an error"), [this]()
+		{
+			FString ExpectedError = TEXT("Failed");
+			TFuture<TValueOrError<FString, FString>> Future =
+				FStringValueOrErrorFuture::MakeError(FString(ExpectedError));
+			if (!TestTrue(TEXT("is ready"), Future.IsReady())) return;
+
+			auto ValueOrError = Future.Get();
+			if (!TestTrue(TEXT("has error"), ValueOrError.HasError())) return;
+
+			TestEqual(TEXT("error"), ValueOrError.GetError(), ExpectedError);
+		});
+	});
+
+	Describe(TEXT("MakeValue"), [this]()
+	{
+		It(TEXT("should return a value"), [this]()
+		{
+			FString ExpectedValue = TEXT("Yo");
+			TFuture<TValueOrError<FString, FString>> Future =
+				FStringValueOrErrorFuture::MakeValue(FString(ExpectedValue));
+			if (!TestTrue(TEXT("is ready"), Future.IsReady())) return;
+
+			auto ValueOrValue = Future.Get();
+			if (!TestTrue(TEXT("has value"), ValueOrValue.HasValue())) return;
+
+			TestEqual(TEXT("value"), ValueOrValue.GetValue(), ExpectedValue);
+		});
+	});
+}
+
+#endif  // WITH_DEV_AUTOMATION_TESTS
diff --git a/Engine/Plugins/Experimental/ToolsetRegistry/Source/ToolsetRegistry/Private/ToolsetRegistry/ValueOrErrorFuture.h b/Engine/Plugins/Experimental/ToolsetRegistry/Source/ToolsetRegistry/Private/ToolsetRegistry/ValueOrErrorFuture.h
new file mode 100644
--- /dev/null
+++ b/Engine/Plugins/Experimental/ToolsetRegistry/Source/ToolsetRegistry/Private/ToolsetRegistry/ValueOrErrorFuture.h
@@ -0,0 +1,34 @@
+// Copyright Epic Games, Inc. All Rights Reserved.
+
+#pragma once
+
+#include "Async/Future.h"
+#include "Containers/UnrealString.h"
+#include "Templates/UnrealTemplate.h"
+#include "Templates/ValueOrError.h"
+
+namespace UE::ToolsetRegistry::Internal
+{
+	// Reduces the boilerplate required to construct a completed future with a value 
+	// or error 
+	template<typename ValueT>
+	class TValueOrErrorFuture
+	{
+	public:
+		using FValueOrError = TValueOrError<ValueT, FString>;
+		using FFuture = TFuture<FValueOrError>;
+
+	public:
+		static FFuture MakeValue(ValueT&& Value)
+		{
+			return MakeFulfilledPromise<FValueOrError>(::MakeValue(MoveTemp(Value))).GetFuture();
+		}
+
+		static FFuture MakeError(FString&& Error)
+		{
+			return MakeFulfilledPromise<FValueOrError>(::MakeError(MoveTemp(Error))).GetFuture();
+		}
+	};
+
+	using FStringValueOrErrorFuture = TValueOrErrorFuture<FString>;
+}
\ No newline at end of file
