[AI Assistant] Move most tool execution code to no longer block.

This moves everything except UToolsetRegistry::ExecuteTool, which requires
CL 50892421, to no longer block when executing tool calls.

Also, since I needed to modify some tests this cleans up some formatting
and deduplicates a load of logic in some tests.

#Synced-CL 51038616

diff --git a/Engine/Plugins/Experimental/ToolsetRegistry/Source/ToolsetRegistry/Private/Tests/BlueprintLibraryToolsetTest.cpp b/Engine/Plugins/Experimental/ToolsetRegistry/Source/ToolsetRegistry/Private/Tests/BlueprintLibraryToolsetTest.cpp
--- a/Engine/Plugins/Experimental/ToolsetRegistry/Source/ToolsetRegistry/Private/Tests/BlueprintLibraryToolsetTest.cpp
+++ b/Engine/Plugins/Experimental/ToolsetRegistry/Source/ToolsetRegistry/Private/Tests/BlueprintLibraryToolsetTest.cpp
@@ -13,15 +13,16 @@
 #include "Serialization/JsonReader.h"
 #include "Serialization/JsonWriter.h"
 #include "Serialization/JsonSerializer.h"
-
-#include "Tests/ToolsetRegistryTestSchemaBuilder.h"
-#include "Tests/ToolsetRegistryTestFlags.h"
+#include "Templates/SharedPointer.h"
 
 #include "BPToolsetRegistry.h"
 #include "ToolsetRegistrySubsystem.h"
 #include "ToolsetRegistry/BlueprintLibraryToolset.h"
 #include "ToolsetRegistry/JsonConversion.h"
 #include "ToolsetRegistry/Toolset.h"
+#include "Tests/ToolsetRegistryTestSchemaBuilder.h"
+#include "Tests/ToolsetRegistryTestFlags.h"
+#include "Tests/UntilFutureCommand.h"
 
 #if WITH_DEV_AUTOMATION_TESTS
 const FString FakeToolsetName(TEXT("ToolsetRegistry.FakeToolset"));
@@ -460,20 +461,25 @@ IMPLEMENT_SIMPLE_AUTOMATION_TEST(
 	"AI.ToolsetRegistry.BlueprintLibraryToolset.ExecuteTool",
 	ToolsetRegistryTest::Flags)
 
-bool FBlueprintLibraryToolsetTestExecuteTool::RunTest(const FString& UnusedParameters)
+	bool FBlueprintLibraryToolsetTestExecuteTool::RunTest(const FString& UnusedParameters)
 {
 	ToolsetRegistryContainer Container(UFakeToolset::StaticClass());
 	TSharedPtr<UE::ToolsetRegistry::FBlueprintLibraryToolset> Toolset = Container.ToolsetHandler();
-	if (!TestTrue("Valid toolset handler returned", Toolset.IsValid())) return false;
-
-	FString Result;
-	FString Error;
-
-	Toolset->InvokeTool(FakeToolName, FakeToolsetInput, Result, Error);
-	TestTrue("Tool execution returns no error", Error.IsEmpty());
-
-	FString ExpectedOutput = FString(TEXT(R"({"returnValue":"hello"})"));
-	TestEqual("Tool execution returns expected output", Result, ExpectedOutput);
+	if (!TestTrue(TEXT("Valid toolset handler returned"), Toolset.IsValid())) return false;
+
+	AddCommand(
+		new UE::ToolsetRegistry::Internal::FUntilFutureCommand<
+				TValueOrError<FString, FString>>(
+			Toolset->ExecuteTool(FakeToolName, FakeToolsetInput),
+			[this](TValueOrError<FString, FString>&& Result) -> void
+			{
+				if (TestTrue(TEXT("Tool has value"), Result.HasValue()))
+				{
+					TestEqual(TEXT("Tool execution returns expected output"),
+						Result.GetValue(),
+						TEXT(R"({"returnValue":"hello"})"));
+				}
+			}));
 	return true;
 }
 
diff --git a/Engine/Plugins/Experimental/ToolsetRegistry/Source/ToolsetRegistry/Private/Tests/ObjectFunctionToolCallTest.cpp b/Engine/Plugins/Experimental/ToolsetRegistry/Source/ToolsetRegistry/Private/Tests/ObjectFunctionToolCallTest.cpp
--- a/Engine/Plugins/Experimental/ToolsetRegistry/Source/ToolsetRegistry/Private/Tests/ObjectFunctionToolCallTest.cpp
+++ b/Engine/Plugins/Experimental/ToolsetRegistry/Source/ToolsetRegistry/Private/Tests/ObjectFunctionToolCallTest.cpp
@@ -8,7 +8,7 @@
 #include "Serialization/JsonWriter.h"
 
 #include "Tests/ToolsetRegistryTestFlags.h"
-
+#include "Tests/UntilFutureCommand.h"
 #include "ToolsetRegistryModule.h"
 #include "ToolsetRegistry/JsonConversion.h"
 #include "ToolsetRegistry/ObjectFunctionToolCall.h"
@@ -37,13 +37,56 @@ namespace UE::ToolsetRegistry
 				UE_LOG(LogToolsetRegistry, Error, TEXT("Json \"%s\" is invalid."), Name);
 			}
 		}
+		
+		// Ensure a tool call result has no error.
+		bool TestHasNoError(FAutomationTestBase& Test, FObjectFunctionToolCall::FToolCallResult& Result)
+		{
+			const FString* Error = Result.TryGetError();
+			return !Error || Test.TestEqual(TEXT("Tool call result has no error."), *Error, TEXT(""));
+		}
+
+		// Ensure a tool call result has a valid value.
+		TSharedPtr<FJsonObject> TestHasValue(
+			FAutomationTestBase& Test, FObjectFunctionToolCall::FToolCallResult& Result)
+		{
+			TSharedPtr<FJsonObject> JsonObject;
+			if (TestHasNoError(Test, Result) &&
+				Test.TestTrue(TEXT("Tool call result has value."), Result.HasValue()))
+			{
+				JsonObject = Result.GetValue();
+				Test.TestTrue(TEXT("Tool call result has valid value."), JsonObject.IsValid());
+#if DEBUG_PRINT_JSONS_TO_LOG
+				if (JsonObject)
+				{
+					PrintJsonToLog(TEXT("Tool call result value"), JsonObject.ToSharedRef());
+				}
+#endif  // DEBUG_PRINT_JSONS_TO_LOG
+			}
+			return JsonObject;
+		}
+
+		// Ensure a tool call result JSON object has a returnValue field.
+		TSharedPtr<FJsonObject> TestHasReturnValueField(
+			FAutomationTestBase& Test, FObjectFunctionToolCall::FToolCallResult& Result,
+			EJson JsonType)	
+		{
+			TSharedPtr<FJsonObject> JsonObject = TestHasValue(Test, Result);
+			if (JsonObject &&
+				!Test.TestTrue(
+					TEXT("Tool call result has returnValue."),
+					JsonObject->HasTypedField(TEXT("returnValue"), JsonType)))
+			{
+				JsonObject.Reset();
+			}
+			return JsonObject;
+		}
 	}
 
 	// Test function with required param.
 
 	IMPLEMENT_SIMPLE_AUTOMATION_TEST(
 		FObjectFunctionToolCallTestExecuteToolCallFunctionWithRequiredParam,
-		"AI.ToolsetRegistry.ToolCall.ExecuteToolCallFunctionWithRequiredParam",
+		"AI.ToolsetRegistry.ObjectFunctionToolCall.ExecuteToolCallFunctionWithRequiredParam",
 		ToolsetRegistryTest::Flags);
 
 	bool FObjectFunctionToolCallTestExecuteToolCallFunctionWithRequiredParam::RunTest(
@@ -66,65 +109,28 @@ namespace UE::ToolsetRegistry
 		// Execute.
 
 		const FObjectFunctionToolCall ToolCall =
-			FObjectFunctionToolCall::ConstructWithNamedClassFunction(InstanceObject.Get(), FunctionName);
+			FObjectFunctionToolCall::ConstructWithNamedClassFunction(
+				InstanceObject.Get(), FunctionName);
 
 		FObjectFunctionToolCall::FFunctionInputParamsJson FunctionInputParamsJson;
 		FunctionInputParamsJson.Set<FString>(FunctionInputParamsJsonString);
-		const FObjectFunctionToolCall::FToolCallResult ToolCallResult =
-			ToolCall.Execute(FunctionInputParamsJson);
-
-		if (!TestTrue("Execute tool call result does not have error.",
-			!ToolCallResult.HasError()))
-		{
-			UE_LOG(LogToolsetRegistry, Error, TEXT("Tool call error - %s"),
-				*ToolCallResult.GetError());
-			return false;
-		}
-		if (!TestTrue("Execute tool call result has value.",
-			ToolCallResult.HasValue()))
-		{
-			return false;
-		}
-
-		const TSharedPtr<FJsonObject> JsonObject = ToolCallResult.GetValue();
-
-		if (!TestTrue("Execute tool call result has valid value.",
-			JsonObject.IsValid()))
-		{
-			return false;
-		}
 
-#if DEBUG_PRINT_JSONS_TO_LOG
-		PrintJsonToLog(TEXT("ToolCallWithRequiredParamReturnedJsonObject"),
-			JsonObject.ToSharedRef());
-#endif
-
-		if (!TestTrue("Return object is not empty (has values.)",
-			JsonObject->Values.Num() > 0))
-		{
-			return false;
-		}
-		if (!TestTrue(TEXT("Return object has float returnValue."),
-			JsonObject->HasTypedField(TEXT("returnValue"), EJson::Number)))
-		{
-			return false;
-		}
-
-		// Check.
-
-		const float ReturnValueFloat =
-			JsonObject->GetNumberField(TEXT("returnValue"));
-		const float ReturnValueFloatExpected =
-			InstanceObject->TestFuncWithRequiredParam(FunctionInputFloatValue);
-		const FString TestMessage = FString::Printf(TEXT("Return value %f is expected value %f."),
-			ReturnValueFloat, ReturnValueFloatExpected);
-		if (!TestTrue(TestMessage,
-			ReturnValueFloat == ReturnValueFloatExpected))
-		{
-			return false;
-		}
-
-		// Success.
+		AddCommand(
+			new UE::ToolsetRegistry::Internal::FUntilFutureCommand<
+					FObjectFunctionToolCall::FToolCallResult>(
+				ToolCall.Execute(FunctionInputParamsJson),
+				[this, InstanceObject](
+					FObjectFunctionToolCall::FToolCallResult&& ToolCallResult) -> void
+				{
+					TSharedPtr<FJsonObject> JsonObject =
+						TestHasReturnValueField(*this, ToolCallResult, EJson::Number);
+					if (!JsonObject) return;
+					TestNearlyEqual(
+						TEXT("Tool call result value's returnValue"),
+						static_cast<float>(JsonObject->GetNumberField(TEXT("returnValue"))),
+						InstanceObject->TestFuncWithRequiredParam(FunctionInputFloatValue),
+						0.0001f);
+				}));
 
 		return true;
 	}
@@ -133,7 +139,7 @@ namespace UE::ToolsetRegistry
 
 	IMPLEMENT_SIMPLE_AUTOMATION_TEST(
 		FObjectFunctionToolCallTestExecuteToolCallFunctionWithOptionalParam,
-		"AI.ToolsetRegistry.ToolCall.ExecuteToolCallFunctionOptionalParam",
+		"AI.ToolsetRegistry.ObjectFunctionToolCall.ExecuteToolCallFunctionOptionalParam",
 		ToolsetRegistryTest::Flags);
 
 	bool FObjectFunctionToolCallTestExecuteToolCallFunctionWithOptionalParam::RunTest(
@@ -154,67 +160,29 @@ namespace UE::ToolsetRegistry
 
 		const TOptional<FObjectFunctionToolCall> ToolCall =
 			FObjectFunctionToolCall::ConstructWithNamedClassFunction(InstanceObject.Get(), FunctionName);
-		if (!TestTrue("Tool call created from class and function name.",
+		if (!TestTrue(TEXT("Tool call created from class and function name."),
 			ToolCall.IsSet()))
 		{
 			return false;
 		}
 
-		const FObjectFunctionToolCall::FToolCallResult ToolCallResult = ToolCall->Execute();
-
-		if (!TestTrue("Execute tool call result does not have error.",
-			!ToolCallResult.HasError()))
-		{
-			UE_LOG(LogToolsetRegistry, Error, TEXT("Tool call error - %s"),
-				*ToolCallResult.GetError());
-			return false;
-		}
-		if (!TestTrue("Execute tool call result has value.",
-			ToolCallResult.HasValue()))
-		{
-			return false;
-		}
-
-		const TSharedPtr<FJsonObject> JsonObject = ToolCallResult.GetValue();
-
-		if (!TestTrue("Execute tool call result has valid value.",
-			JsonObject.IsValid()))
-		{
-			return false;
-		}
-
-#if DEBUG_PRINT_JSONS_TO_LOG
-		PrintJsonToLog(TEXT("ToolCallWithOptionalParamReturnedJsonObject"),
-			JsonObject.ToSharedRef());
-#endif
-
-		if (!TestTrue("Return object is not empty (has values.)",
-			JsonObject->Values.Num() > 0))
-		{
-			return false;
-		}
-		if (!TestTrue(TEXT("Return object has float returnValue."),
-			JsonObject->HasTypedField(TEXT("returnValue"), EJson::Number)))
-		{
-			return false;
-		}
-
-		// Check.
-
-		const float ReturnValueFloat =
-			JsonObject->GetNumberField(TEXT("returnValue"));
-		const float ReturnValueFloatExpected =
-			InstanceObject->TestFuncWithOptionalParam();
-		const FString TestMessage =
-			FString::Printf(TEXT("Return value %f is expected value %f."), ReturnValueFloat, ReturnValueFloatExpected);
-		if (!TestTrue(TestMessage,
-			ReturnValueFloat == ReturnValueFloatExpected))
-		{
-			return false;
-		}
-
-		// Success.
-
+		AddCommand(
+			new UE::ToolsetRegistry::Internal::FUntilFutureCommand<
+					FObjectFunctionToolCall::FToolCallResult>(
+				ToolCall->Execute(),
+				[this, InstanceObject](
+					FObjectFunctionToolCall::FToolCallResult&& ToolCallResult) -> void
+				{
+					TSharedPtr<FJsonObject> JsonObject =
+						TestHasReturnValueField(*this, ToolCallResult, EJson::Number);
+					if (!JsonObject) return;
+
+					TestNearlyEqual(
+						TEXT("Tool call result value's returnValue"),
+						static_cast<float>(JsonObject->GetNumberField(TEXT("returnValue"))),
+						InstanceObject->TestFuncWithOptionalParam(),
+						0.0001f);
+				}));
 		return true;
 	}
 
diff --git a/Engine/Plugins/Experimental/ToolsetRegistry/Source/ToolsetRegistry/Private/Tests/UntilFutureCommand.h b/Engine/Plugins/Experimental/ToolsetRegistry/Source/ToolsetRegistry/Private/Tests/UntilFutureCommand.h
--- a/Engine/Plugins/Experimental/ToolsetRegistry/Source/ToolsetRegistry/Private/Tests/UntilFutureCommand.h
+++ b/Engine/Plugins/Experimental/ToolsetRegistry/Source/ToolsetRegistry/Private/Tests/UntilFutureCommand.h
@@ -1,5 +1,7 @@
 // Copyright Epic Games, Inc. All Rights Reserved.
 
+#pragma once
+
 #include "Async/Future.h"
 #include "Misc/AutomationTest.h"
 #include "Templates/Function.h"
diff --git a/Engine/Plugins/Experimental/ToolsetRegistry/Source/ToolsetRegistry/Private/ToolsetRegistry/BlueprintLibraryToolset.cpp b/Engine/Plugins/Experimental/ToolsetRegistry/Source/ToolsetRegistry/Private/ToolsetRegistry/BlueprintLibraryToolset.cpp
--- a/Engine/Plugins/Experimental/ToolsetRegistry/Source/ToolsetRegistry/Private/ToolsetRegistry/BlueprintLibraryToolset.cpp
+++ b/Engine/Plugins/Experimental/ToolsetRegistry/Source/ToolsetRegistry/Private/ToolsetRegistry/BlueprintLibraryToolset.cpp
@@ -3,6 +3,7 @@
 #include "ToolsetRegistry/BlueprintLibraryToolset.h"
 
 #include "Async/Future.h"
+#include "Blueprint/BlueprintExceptionInfo.h"
 #include "Containers/UnrealString.h"
 #include "Dom/JsonObject.h"
 #include "Kismet/BlueprintFunctionLibrary.h"
@@ -11,13 +12,14 @@
 #include "Policies/CondensedJsonPrintPolicy.h"
 #include "Serialization/JsonWriter.h"
 #include "Serialization/JsonSerializer.h"
-#include "Blueprint/BlueprintExceptionInfo.h"
-#include "ToolsetRegistryDelegateHandle.h"
+#include "Templates/UnrealTemplate.h"
 
+#include "ToolsetRegistryDelegateHandle.h"
 #include "ToolsetRegistryModule.h"
 #include "ToolsetRegistry/JsonConversion.h"
 #include "ToolsetRegistry/ObjectFunctionToolCall.h"
 #include "ToolsetRegistry/ToolsetJson.h"
+#include "ToolsetRegistry/ValueOrErrorFuture.h"
 
 namespace UE::ToolsetRegistry
 {
@@ -147,24 +149,6 @@ namespace UE::ToolsetRegistry
 			? FBlueprintLibraryToolset::GetToolsetClassName(LibraryClass) : FString();
 	}
 
-	TFuture<TValueOrError<FString, FString>> FBlueprintLibraryToolset::ExecuteTool(
-		const FString& ToolName, const FString& JsonInput)
-	{
-		FString ToolResult;
-		FString ToolError;
-		InvokeTool(ToolName, JsonInput, ToolResult, ToolError);
-		if (ToolError.IsEmpty())
-		{
-			return MakeFulfilledPromise<TValueOrError<FString, FString>>(
-				MakeValue(ToolResult)).GetFuture();
-		}
-		else
-		{
-			return MakeFulfilledPromise<TValueOrError<FString, FString>>(
-				MakeError(ToolError)).GetFuture();
-		}
-	}
-
 	FString FBlueprintLibraryToolset::GetToolsetJsonSchema() const
 	{
 		TSharedPtr<FJsonObject> Schema = GenerateSchema();
@@ -228,23 +212,24 @@ namespace UE::ToolsetRegistry
 		}
 	}
 
-	void FBlueprintLibraryToolset::InvokeTool(const FString& ToolName, const FString& JsonInput, FString& OutResult, FString& OutError)
+	TFuture<TValueOrError<FString, FString>> FBlueprintLibraryToolset::ExecuteTool(
+		const FString& ToolName, const FString& JsonInput)
 	{
-		OutResult.Empty();
-		OutError.Empty();
+		using namespace UE::ToolsetRegistry::Internal;
+
 		TStrongObjectPtr<UClass> Toolset = ToolsetClass.Pin();
 		if (!Toolset)
 		{
-			OutError = FString(TEXT("Invalid Toolset Class"));
-			return;
+			return FStringValueOrErrorFuture::MakeError(TEXT("Invalid Toolset Class"));
 		}
 
 		TSharedPtr<FObjectFunctionToolCall>* ToolCallPtr = Tools.Find(ToolName);
 		if (!ToolCallPtr)
 		{
-			OutError = FString::Printf(TEXT("Unknown tool %s"), *ToolName);
-			return;
+			return FStringValueOrErrorFuture::MakeError(
+				FString::Printf(TEXT("Unknown tool %s"), *ToolName));
 		}
+
 		TSharedPtr<FObjectFunctionToolCall> ToolCall = *ToolCallPtr;
 		check(ToolCall);
 
@@ -253,30 +238,45 @@ namespace UE::ToolsetRegistry
 
 		// Capture script errors raised by the tool during execution. Note: script errors raised by
 		// tools running outside of the game thread will not be captured correctly.
-		FToolCallScriptExceptionHandler ExceptionHandler;
+		auto ExceptionHandler = MakeShared<FToolCallScriptExceptionHandler>();
 		ensureAlways(IsInGameThread());
 
-		FObjectFunctionToolCall::FToolCallResult ToolCallResult = ToolCall->Execute(JsonParams);
-		if (ToolCallResult.HasError())
-		{
-			OutError = ToolCallResult.StealError();
-			return;
-		}
-		else if (FString ScriptException = ExceptionHandler.GetException();
-		         !ScriptException.IsEmpty())
-		{
-			// If the tool raised an error, return that.
-			OutError = ScriptException;
-			return;
-		}
-		check(ToolCallResult.HasValue());
-		TSharedPtr<FJsonObject> ResultValue = ToolCallResult.StealValue();
-		if (!ResultValue.IsValid())
-		{
-			OutError = FString::Printf(TEXT("Unable to serialize result from executing tool %s"), *ToolName);
-			return;
-		}
-		OutResult = UE::ToolsetRegistry::Internal::JsonToString(ResultValue.ToSharedRef());
+		FString ToolNameCopy = ToolName;
+		return ToolCall->Execute(JsonParams).Next(
+			[
+				Toolset,  // Capture to keep alive.
+				ToolCall = ToolCall,
+				ExceptionHandler = ExceptionHandler,
+				ToolNameCopy = ToolNameCopy
+			](
+				FObjectFunctionToolCall::FToolCallResult&& ToolCallResult) ->
+					TValueOrError<FString, FString>
+			{
+				ensureAlways(IsInGameThread());
+				if (ToolCallResult.HasError())
+				{
+					return MakeError(ToolCallResult.StealError());
+				}
+				else if (
+					FString ScriptException = ExceptionHandler->GetException();
+					!ScriptException.IsEmpty())
+				{
+					// If the tool raised an error, return that.
+					return MakeError(MoveTemp(ScriptException));
+				}
+				check(ToolCallResult.HasValue());
+				TSharedPtr<FJsonObject> ResultValue = ToolCallResult.StealValue();
+				if (!ResultValue.IsValid())
+				{
+					return MakeError(
+						FString::Printf(
+							TEXT("Unable to serialize result from executing tool %s"),
+							*ToolNameCopy));
+				}
+				return MakeValue(
+					UE::ToolsetRegistry::Internal::JsonToString(ResultValue.ToSharedRef()));
+			});
+
 	}
 
 	TArray<FString> FBlueprintLibraryToolset::GetToolNames()
diff --git a/Engine/Plugins/Experimental/ToolsetRegistry/Source/ToolsetRegistry/Private/ToolsetRegistry/ObjectFunctionToolCall.cpp b/Engine/Plugins/Experimental/ToolsetRegistry/Source/ToolsetRegistry/Private/ToolsetRegistry/ObjectFunctionToolCall.cpp
--- a/Engine/Plugins/Experimental/ToolsetRegistry/Source/ToolsetRegistry/Private/ToolsetRegistry/ObjectFunctionToolCall.cpp
+++ b/Engine/Plugins/Experimental/ToolsetRegistry/Source/ToolsetRegistry/Private/ToolsetRegistry/ObjectFunctionToolCall.cpp
@@ -8,10 +8,13 @@
 #include "Engine/EngineTypes.h"
 #include "Serialization/JsonReader.h"
 #include "Serialization/JsonSerializer.h"
+#include "Templates/UnrealTemplate.h"
 #include "ToolsetJson.h"
 #include "UObject/FieldIterator.h"
 
 #include "ToolsetRegistry/JsonConversion.h"
+#include "ToolsetRegistry/ToolCallAsyncResult.h"
+#include "ToolsetRegistry/ValueOrErrorFuture.h"
 
 namespace UE::ToolsetRegistry
 {
@@ -272,21 +275,21 @@ namespace UE::ToolsetRegistry
 		return FObjectFunctionToolCall(InInstanceObject, Function);
 	}
 
-	FObjectFunctionToolCall::FToolCallResult FObjectFunctionToolCall::Execute(
+	TFuture<FObjectFunctionToolCall::FToolCallResult> FObjectFunctionToolCall::Execute(
 		const TOptional<FFunctionInputParamsJson>& FunctionInputParamsJson) const
 	{
+		using FJsonObjectOrErrorFuture = TValueOrErrorFuture<TSharedPtr<FJsonObject>>;
+
 		// Get function input params JSON as a JSON object.
 		TSharedPtr<FJsonObject> FunctionInputParamsJsonObject;
-		if (const TValueOrError<const TSharedRef<FJsonObject>, const FString> JsonObjectOrError =
+		TValueOrError<TSharedRef<FJsonObject>, FString> JsonObjectOrError =
 			BuildValidFunctionInputParamsJsonObject(FunctionInputParamsJson);
-			JsonObjectOrError.HasError())
-		{
-			return MakeError(JsonObjectOrError.GetError());
-		}
-		else
+		if (JsonObjectOrError.HasError())
 		{
-			FunctionInputParamsJsonObject = JsonObjectOrError.GetValue();
+			auto Error = JsonObjectOrError.StealError();
+			return FJsonObjectOrErrorFuture::MakeError(MoveTemp(Error));
 		}
+		FunctionInputParamsJsonObject = JsonObjectOrError.GetValue();
 
 		// Allocate the frame memory.
 		uint8* FrameMemory;
@@ -327,7 +330,7 @@ namespace UE::ToolsetRegistry
 				FunctionInputParamsJsonObject.ToSharedRef(), Function, FrameMemory))
 		{
 			// Caller error. Ill-formed incoming function input params JSON. 
-			return MakeError(FString::Printf(TEXT(
+			return FJsonObjectOrErrorFuture::MakeError(FString::Printf(TEXT(
 				"Function \"%s\", "
 				"could not convert incoming function input params Json to a UStruct. "
 				"Incoming function input params Json may be incorrect or contain unsupported "
@@ -343,19 +346,24 @@ namespace UE::ToolsetRegistry
 		MaybeSetWorldContextPropertyInFrameMemory(Function, FrameMemory);
 
 		// Invoke function.
-		// TODO: A lot of this logic is duplicated from ScriptCore and should be ported to use existing function invocation logic.
+		// TODO: A lot of this logic is duplicated from ScriptCore and should be ported to use 
+		// existing function invocation logic.
 		FFrame StackFrame(InstanceObject, Function, FrameMemory, nullptr,
 			Function->ChildProperties);
 
-		// If the function has out parameters, fill the stack frame's out parameter info with the info for those params 
+		// If the function has out parameters, fill the stack frame's out parameter info with the 
+		// info for those params.
 		if (Function->HasAnyFunctionFlags(FUNC_HasOutParms))
 		{
 			FOutParmRec** LastOut = &StackFrame.OutParms;
-			for (FProperty* Property = (FProperty*)(Function->ChildProperties); Property && (Property->PropertyFlags & (CPF_Parm)) == CPF_Parm; Property = (FProperty*)Property->Next)
+			for (FProperty* Property = (FProperty*)(Function->ChildProperties);
+				Property && (Property->PropertyFlags & (CPF_Parm)) == CPF_Parm;
+				Property = (FProperty*)Property->Next)
 			{
-				// this is used for optional parameters - the destination address for out parameter values is the address of the calling function
-				// so we'll need to know which address to use if we need to evaluate the default parm value expression located in the new function's
-				// bytecode
+				// This is used for optional parameters - the destination address for out parameter
+				// values is  the address of the calling function so we'll need to know which 
+				// address to use if we need to evaluate the default param value expression located
+				// in the new function's bytecode.
 				if (Property->HasAnyPropertyFlags(CPF_OutParm))
 				{
 					CA_SUPPRESS(6263)
@@ -400,12 +408,9 @@ namespace UE::ToolsetRegistry
 
 		if (!ErrorMessage.IsEmpty())
 		{
-			return MakeError(ErrorMessage);
-		}
-		else
-		{
-			return MakeValue(ReturnJsonObject);
+			return FJsonObjectOrErrorFuture::MakeError(MoveTemp(ErrorMessage));
 		}
+		return FJsonObjectOrErrorFuture::MakeValue(MoveTemp(ReturnJsonObject));
 	}
 
 	TSharedPtr<FJsonObject> FObjectFunctionToolCall::MakeReturnObject(
@@ -428,7 +433,7 @@ namespace UE::ToolsetRegistry
 		return ReturnJsonObject;
 	}
 
-	TValueOrError<const TSharedRef<FJsonObject>, const FString>
+	TValueOrError<TSharedRef<FJsonObject>, FString>
 		FObjectFunctionToolCall::BuildValidFunctionInputParamsJsonObject(
 			const TOptional<FFunctionInputParamsJson>& FunctionInputParamsJson) const
 	{
diff --git a/Engine/Plugins/Experimental/ToolsetRegistry/Source/ToolsetRegistry/Private/ToolsetRegistry/ObjectFunctionToolCall.h b/Engine/Plugins/Experimental/ToolsetRegistry/Source/ToolsetRegistry/Private/ToolsetRegistry/ObjectFunctionToolCall.h
--- a/Engine/Plugins/Experimental/ToolsetRegistry/Source/ToolsetRegistry/Private/ToolsetRegistry/ObjectFunctionToolCall.h
+++ b/Engine/Plugins/Experimental/ToolsetRegistry/Source/ToolsetRegistry/Private/ToolsetRegistry/ObjectFunctionToolCall.h
@@ -2,12 +2,13 @@
 
 #pragma once
 
-#include "Templates/SharedPointer.h"
+#include "Async/Future.h"
+#include "Dom/JsonObject.h"
 #include "Misc/NotNull.h"
 #include "Misc/TVariant.h"
+#include "Templates/SharedPointer.h"
 #include "Templates/ValueOrError.h"
 #include "UObject/Class.h"
-#include "Dom/JsonObject.h"
 
 //
 // UE::ToolsetRegistry::FObjectFunctionToolCall
@@ -57,17 +58,18 @@ namespace UE::ToolsetRegistry
 		 *		this must be supplied with JSON that provides those required params. Otherwise, an
 		 *		error will result. If the function this class was initialized with has no required 
 		 *		params, then this does not need to be supplied.
-		 * @return Tool call result. (See above). When tool result does not have an error -
+		 * @return Tool call result future. (See above). When tool result does not have an error -
  		 *		Returned objects will either be (1) the object itself, or (2) a value named 
  		 *		'returnValue' inside an object. And this may be null if (1) function does not 
  		 *		return a value, or (2) if function does return a value and that return value 
  		 *		is null.
 		 */
-		FToolCallResult Execute(const TOptional<FFunctionInputParamsJson>& FunctionInputParamsJson = 
-			TOptional<FFunctionInputParamsJson>()) const;
+		TFuture<FToolCallResult> Execute(
+			const TOptional<FFunctionInputParamsJson>& FunctionInputParamsJson = 
+				TOptional<FFunctionInputParamsJson>()) const;
 
 	private:
-		TValueOrError<const TSharedRef<FJsonObject>, const FString> 
+		TValueOrError<TSharedRef<FJsonObject>, FString> 
 		BuildValidFunctionInputParamsJsonObject(
 			const TOptional<FFunctionInputParamsJson>& FunctionInputParamsJson) const;
 
diff --git a/Engine/Plugins/Experimental/ToolsetRegistry/Source/ToolsetRegistry/Public/ToolsetRegistry/BlueprintLibraryToolset.h b/Engine/Plugins/Experimental/ToolsetRegistry/Source/ToolsetRegistry/Public/ToolsetRegistry/BlueprintLibraryToolset.h
--- a/Engine/Plugins/Experimental/ToolsetRegistry/Source/ToolsetRegistry/Public/ToolsetRegistry/BlueprintLibraryToolset.h
+++ b/Engine/Plugins/Experimental/ToolsetRegistry/Source/ToolsetRegistry/Public/ToolsetRegistry/BlueprintLibraryToolset.h
@@ -4,9 +4,14 @@
 
 #include "CoreMinimal.h"
 
+#include "Async/Future.h"
+#include "Containers/Map.h"
+#include "Containers/UnrealString.h"
 #include "Dom/JsonObject.h"
-#include "Kismet/BlueprintFunctionLibrary.h"
 #include "Templates/SharedPointer.h"
+#include "Templates/ValueOrError.h"
+#include "UObject/Class.h"
+#include "UObject/WeakObjectPtrTemplates.h"
 
 #include "ToolsetRegistry/Toolset.h"
 
@@ -19,7 +24,7 @@ namespace UE::ToolsetRegistry
 	// Wraps a Blueprint Function Library for use as a Toolset.
 	class FBlueprintLibraryToolset : public FToolset
 	{
-	  public:
+	public:
 		UE_API FBlueprintLibraryToolset(UClass* LibraryClass);
 
 		// Executes a tool with the given name and JSON input, returning a future with the JSON
@@ -49,18 +54,6 @@ namespace UE::ToolsetRegistry
 			return ToolsetClass.Pin() ? ToolsetClass->GetToolTipText().ToString() : FString();
 		}
 
-		// Invoke a tool in this toolset by name.
-		//
-		// Expected to be called by the toolset registry ExecuteTool method, so
-		// ToolName should be the unqualified function name.
-		//
-		// Returns JSON result in OutResult unless there was an error.
-		//
-		// Returns error message in OutError if there was error in invocation.
-		UE_API void InvokeTool(
-			const FString& ToolName, const FString& JsonInput,
-			FString& OutResult, FString& OutError);
-
 		UE_API TArray<FString> GetToolNames();
 
 		// Whether the toolset has a valid set of tools.
@@ -69,8 +62,7 @@ namespace UE::ToolsetRegistry
 		// Get qualified name for toolset class.
 		UE_API static FString GetToolsetClassName(UClass* InClass);
 
-	  private:
-
+	private:
 		// The wrapped toolset function library
 		TWeakObjectPtr<UClass> ToolsetClass;
 
