[AI Assistant] Fix FUntilCommand and FModifyAndValidateAsyncResultCommand

My initial attempt to subclass FUntilCommand failed as the automation framework
calls the private IAutomationLatentCommand::InternalUpdate() rather than the
overridable Update() method. This causes both subclasses to timeout immediately
and tests to fail.

This fixes the issue by doing away with a subclass and instead using a struct to
hold the state that the closure passed to FUnitCommand needs to reference to
extend it.

#Synced-CL 51043000

#Synced-CL 51043000

diff --git a/Engine/Plugins/Experimental/ToolsetRegistry/Source/ToolsetRegistry/Private/Tests/BlueprintLibraryToolsetTest.cpp b/Engine/Plugins/Experimental/ToolsetRegistry/Source/ToolsetRegistry/Private/Tests/BlueprintLibraryToolsetTest.cpp
--- a/Engine/Plugins/Experimental/ToolsetRegistry/Source/ToolsetRegistry/Private/Tests/BlueprintLibraryToolsetTest.cpp
+++ b/Engine/Plugins/Experimental/ToolsetRegistry/Source/ToolsetRegistry/Private/Tests/BlueprintLibraryToolsetTest.cpp
@@ -468,8 +468,8 @@ IMPLEMENT_SIMPLE_AUTOMATION_TEST(
 	if (!TestTrue(TEXT("Valid toolset handler returned"), Toolset.IsValid())) return false;
 
 	AddCommand(
-		new UE::ToolsetRegistry::Internal::FUntilFutureCommand<
-				TValueOrError<FString, FString>>(
+		UE::ToolsetRegistry::Internal::FUntilFutureCommand<
+				TValueOrError<FString, FString>>::Create(
 			Toolset->ExecuteTool(FakeToolName, FakeToolsetInput),
 			[this](TValueOrError<FString, FString>&& Result) -> void
 			{
diff --git a/Engine/Plugins/Experimental/ToolsetRegistry/Source/ToolsetRegistry/Private/Tests/ObjectFunctionToolCallTest.cpp b/Engine/Plugins/Experimental/ToolsetRegistry/Source/ToolsetRegistry/Private/Tests/ObjectFunctionToolCallTest.cpp
--- a/Engine/Plugins/Experimental/ToolsetRegistry/Source/ToolsetRegistry/Private/Tests/ObjectFunctionToolCallTest.cpp
+++ b/Engine/Plugins/Experimental/ToolsetRegistry/Source/ToolsetRegistry/Private/Tests/ObjectFunctionToolCallTest.cpp
@@ -116,8 +116,8 @@ namespace UE::ToolsetRegistry
 		FunctionInputParamsJson.Set<FString>(FunctionInputParamsJsonString);
 
 		AddCommand(
-			new UE::ToolsetRegistry::Internal::FUntilFutureCommand<
-					FObjectFunctionToolCall::FToolCallResult>(
+			UE::ToolsetRegistry::Internal::FUntilFutureCommand<
+					FObjectFunctionToolCall::FToolCallResult>::Create(
 				ToolCall.Execute(FunctionInputParamsJson),
 				[this, InstanceObject](
 					FObjectFunctionToolCall::FToolCallResult&& ToolCallResult) -> void
@@ -167,8 +167,8 @@ namespace UE::ToolsetRegistry
 		}
 
 		AddCommand(
-			new UE::ToolsetRegistry::Internal::FUntilFutureCommand<
-					FObjectFunctionToolCall::FToolCallResult>(
+			UE::ToolsetRegistry::Internal::FUntilFutureCommand<
+					FObjectFunctionToolCall::FToolCallResult>::Create(
 				ToolCall->Execute(),
 				[this, InstanceObject](
 					FObjectFunctionToolCall::FToolCallResult&& ToolCallResult) -> void
diff --git a/Engine/Plugins/Experimental/ToolsetRegistry/Source/ToolsetRegistry/Private/Tests/ToolCallAsyncResultTest.cpp b/Engine/Plugins/Experimental/ToolsetRegistry/Source/ToolsetRegistry/Private/Tests/ToolCallAsyncResultTest.cpp
--- a/Engine/Plugins/Experimental/ToolsetRegistry/Source/ToolsetRegistry/Private/Tests/ToolCallAsyncResultTest.cpp
+++ b/Engine/Plugins/Experimental/ToolsetRegistry/Source/ToolsetRegistry/Private/Tests/ToolCallAsyncResultTest.cpp
@@ -240,57 +240,55 @@ namespace
 	}
 
 	// Modify an async result in a latent command and then validate its state.
-	class FModifyAndValidateAsyncResultCommand : public IAutomationLatentCommand
+	class FModifyAndValidateAsyncResultCommand
 	{
+	private:
+		struct FState
+		{
+			TSharedPtr<FToolCallAsyncResultCompletionValidator> Validator;
+			TFunction<bool()> ModifyResult;
+			bool bHasModifiedResult = false;
+		};
+
 	public:
-		FModifyAndValidateAsyncResultCommand(
+		static FUntilCommand* Create(
 			TSharedPtr<FToolCallAsyncResultCompletionValidator> InValidator,
-			TFunction<bool()>&& InModifyResult) :
-			Validator(InValidator),
-			ModifyResult(InModifyResult),
-			UntilCommand(
-				[this]() -> bool
+			TFunction<bool()>&& InModifyResult)
+		{
+			auto State = MakeShared<FState>();
+			State->Validator = InValidator;
+			State->ModifyResult = InModifyResult;
+			return new FUntilCommand(
+				[State]() -> bool
 				{
-					if (!bHasModifiedResult)
+					if (!State->bHasModifiedResult)
 					{
-						if (!ModifyResult()) return true;
-						bHasModifiedResult = true;
+						if (!State->ModifyResult()) return true;
+						State->bHasModifiedResult = true;
 					}
 
 					// Wait for the async result to complete.
-					if (!Validator->GetResultHandlerFuture().IsReady()) return false;
+					if (!State->Validator->GetResultHandlerFuture().IsReady()) return false;
 
-					Validator->Validate();
+					State->Validator->Validate();
 					return true;
-				})
-		{
+				});
 		}
 
-		FModifyAndValidateAsyncResultCommand(
-			TSharedPtr<FToolCallAsyncResultCompletionValidator> InValidator,
+		static FUntilCommand* Create(
+			TSharedPtr<FToolCallAsyncResultCompletionValidator> Validator,
 			TSharedPtr<FExecuteOnThreadHandle<bool>> ExecuteOnThreadHandle,
-			FString ExecuteOnThreadTaskDescription) :
-			FModifyAndValidateAsyncResultCommand(
-				InValidator,
-				[this, ExecuteOnThreadHandle, ExecuteOnThreadTaskDescription]() -> bool
+			FString ExecuteOnThreadTaskDescription)
+		{
+			return Create(
+				Validator,
+				[ExecuteOnThreadHandle, ExecuteOnThreadTaskDescription, Validator]() -> bool
 				{
 					ExecuteOnThreadHandle->Start();
 					return Validator->GetTestCase().TestTrue(
 						*ExecuteOnThreadTaskDescription, ExecuteOnThreadHandle->ReturnValue());
-				})
-		{
+				});
 		}
-
-		virtual bool Update() override
-		{
-			return UntilCommand.Update();
-		}
-
-	private:
-		TSharedPtr<FToolCallAsyncResultCompletionValidator> Validator;
-		TFunction<bool()> ModifyResult;
-		bool bHasModifiedResult = false;
-		FUntilCommand UntilCommand;
 	};
 
 }  // namespace
@@ -363,7 +361,7 @@ void FToolCallAsyncResultBaseSpec::Define()
 		Validator->ExpectError(FakeError);
 
 		AddCommand(
-			new FModifyAndValidateAsyncResultCommand(
+			FModifyAndValidateAsyncResultCommand::Create(
 				Validator,
 				ExecuteOnThread([this]() -> bool { return Result->SetError(FakeError); }),
 				TEXT("Scheduled set error")));
@@ -376,7 +374,7 @@ void FToolCallAsyncResultBaseSpec::Define()
 		Validator->ExpectError(FakeError);
 
 		AddCommand(
-			new FModifyAndValidateAsyncResultCommand(
+			FModifyAndValidateAsyncResultCommand::Create(
 				Validator,
 				ExecuteOnThread(
 					[this]() -> bool
@@ -434,7 +432,7 @@ void FToolCallAsyncResultBaseSpec::Define()
 		Validator->ExpectJsonValue(ExpectedJsonValue);
 
 		AddCommand(
-			new FModifyAndValidateAsyncResultCommand(
+			FModifyAndValidateAsyncResultCommand::Create(
 				Validator,
 				ExecuteOnThread([this]() -> bool { return SetResultValue(); }),
 				TEXT("Schedule set value")));
@@ -448,7 +446,7 @@ void FToolCallAsyncResultBaseSpec::Define()
 		Validator->ExpectJsonValue(ExpectedJsonValue);
 
 		AddCommand(
-			new FModifyAndValidateAsyncResultCommand(
+			FModifyAndValidateAsyncResultCommand::Create(
 				Validator,
 				ExecuteOnThread(
 					[this]() -> bool
diff --git a/Engine/Plugins/Experimental/ToolsetRegistry/Source/ToolsetRegistry/Private/Tests/UntilFutureCommand.h b/Engine/Plugins/Experimental/ToolsetRegistry/Source/ToolsetRegistry/Private/Tests/UntilFutureCommand.h
--- a/Engine/Plugins/Experimental/ToolsetRegistry/Source/ToolsetRegistry/Private/Tests/UntilFutureCommand.h
+++ b/Engine/Plugins/Experimental/ToolsetRegistry/Source/ToolsetRegistry/Private/Tests/UntilFutureCommand.h
@@ -13,38 +13,36 @@ namespace UE::ToolsetRegistry::Internal
 {
 	// Latent test command that waits until a future is complete or a timeout elapses.
 	template<typename ResultT>
-	class FUntilFutureCommand : public IAutomationLatentCommand
+	class FUntilFutureCommand
 	{
+	private:
+		struct FState
+		{
+			TFuture<ResultT> WaitOnFuture;
+			TFunction<void(ResultT&&)> OnFutureCompletion;
+		};
+
 	public:
-		FUntilFutureCommand(
+		static FUntilCommand* Create(
 			TFuture<ResultT>&& InWaitOnFuture,
 			TFunction<void(ResultT&&)> InFutureCompletion,
-			float InTimeout = 5.0f) :
-			WaitOnFuture(MoveTemp(InWaitOnFuture)),
-			OnFutureCompletion(InFutureCompletion),
-			UntilCommand(
-				[this]() -> bool
+			float InTimeout = 5.0f)
+		{
+			auto State = MakeShared<FState>();
+			State->WaitOnFuture = MoveTemp(InWaitOnFuture);
+			State->OnFutureCompletion = InFutureCompletion;
+			return new FUntilCommand(
+				[State]() -> bool
 				{
-					bool bIsReady = WaitOnFuture.IsReady();
+					bool bIsReady = State->WaitOnFuture.IsReady();
 					if (bIsReady)
 					{
-						OnFutureCompletion(WaitOnFuture.Consume());
+						State->OnFutureCompletion(State->WaitOnFuture.Consume());
 					}
 					return bIsReady;
 				},
-				InTimeout)
-		{
+				InTimeout);
 		}
-
-		virtual bool Update() override
-		{
-			return UntilCommand.Update();
-		}
-
-	private:
-		TFuture<ResultT> WaitOnFuture;
-		TFunction<void(ResultT&&)> OnFutureCompletion;
-		FUntilCommand UntilCommand;
 	};
 }
 
diff --git a/Engine/Plugins/Experimental/ToolsetRegistry/Source/ToolsetRegistry/Private/Tests/UntilFutureCommandTest.cpp b/Engine/Plugins/Experimental/ToolsetRegistry/Source/ToolsetRegistry/Private/Tests/UntilFutureCommandTest.cpp
--- a/Engine/Plugins/Experimental/ToolsetRegistry/Source/ToolsetRegistry/Private/Tests/UntilFutureCommandTest.cpp
+++ b/Engine/Plugins/Experimental/ToolsetRegistry/Source/ToolsetRegistry/Private/Tests/UntilFutureCommandTest.cpp
@@ -35,7 +35,7 @@ void FUnitFutureCommandSpec::Define()
 					},
 					0.1f /* wait 100ms */));
 			AddCommand(
-				new FUntilFutureCommand<FString>(
+				FUntilFutureCommand<FString>::Create(
 					Promise->GetFuture(),
 					[this, ExpectedResult](FString&& Result) -> void
 					{
