ToolsetLibrary improvements:
* Moved properties related functions out of ToolsetLibraryTemp and into ToolsetLibrary
* Added a GetDerivedClasses to ToolsetLibrary.
* Added tests for ToolsetLibrary.

#Synced-CL 51002830

diff --git a/Engine/Plugins/Experimental/ToolsetRegistry/Source/ToolsetRegistry/Private/AgentSkill.cpp b/Engine/Plugins/Experimental/ToolsetRegistry/Source/ToolsetRegistry/Private/AgentSkill.cpp
--- a/Engine/Plugins/Experimental/ToolsetRegistry/Source/ToolsetRegistry/Private/AgentSkill.cpp
+++ b/Engine/Plugins/Experimental/ToolsetRegistry/Source/ToolsetRegistry/Private/AgentSkill.cpp
@@ -3,13 +3,11 @@
 #include "AgentSkill.h"
 
 #include "AssetRegistry/AssetData.h"
-#include "AssetRegistry/AssetRegistryHelpers.h"
-#include "AssetRegistry/AssetRegistryModule.h"
 #include "Interfaces/ITargetPlatform.h"
 
 #include "ToolsetRegistry/BlueprintLibraryToolset.h"
-#include "ToolsetRegistryModule.h"
 #include "ToolsetRegistry/Toolset.h"
+#include "ToolsetLibrary.h"
 
 #include UE_INLINE_GENERATED_CPP_BY_NAME(AgentSkill)
 
@@ -28,34 +26,15 @@ FAgentSkillDetails UAgentSkill::GetDetails() const
 }
 
 TMap<FString, FString> UAgentSkillToolset::ListSkills()
-{	
-	FAssetRegistryModule& AssetRegistryModule = FModuleManager::LoadModuleChecked<FAssetRegistryModule>("AssetRegistry");
-	IAssetRegistry& AssetRegistry = AssetRegistryModule.Get();
-
-	TSet<FTopLevelAssetPath> DerivedNames;
-	AssetRegistry.GetDerivedClassNames({ UAgentSkill::StaticClass()->GetClassPathName() }, {}, DerivedNames);
-
+{
 	TMap<FString, FString> Summaries;
-	for (const FTopLevelAssetPath& DerivedName : DerivedNames)
+	TArray<UClass*> DerivedClasses = UToolsetLibrary::GetDerivedClasses(UAgentSkill::StaticClass());
+	for (const UClass* DerivedClass : DerivedClasses)
 	{
-		// We ignore skeleton classes (contained in BPs) but consider all other AgentSkill subclasses.
-		if (!DerivedName.GetAssetName().ToString().StartsWith(TEXT("SKEL_")))
+		UAgentSkill* AgentSkill = Cast<UAgentSkill>(DerivedClass->GetDefaultObject());
+		if (AgentSkill && !AgentSkill->Description.IsEmpty())
 		{
-			FString PathName = DerivedName.ToString();
-			UClass* LoadedClass = FSoftClassPath(PathName).TryLoadClass<UAgentSkill>();
-			if (LoadedClass)
-			{
-				UAgentSkill* AgentSkill = Cast<UAgentSkill>(LoadedClass->GetDefaultObject());
-				if (AgentSkill && !AgentSkill->Description.IsEmpty())
-				{
-					Summaries.Add(PathName, AgentSkill->Description);
-				}
-			}
-			else
-			{
-				UE_LOG(LogToolsetRegistry, Warning,
-					   TEXT("Unable to load AgentSkill subclass: %s"), *DerivedName.ToString());
-			}
+			Summaries.Add(DerivedClass->GetPathName(), AgentSkill->Description);
 		}
 	}
 
diff --git a/Engine/Plugins/Experimental/ToolsetRegistry/Source/ToolsetRegistry/Private/Tests/ToolsetLibraryTest.cpp b/Engine/Plugins/Experimental/ToolsetRegistry/Source/ToolsetRegistry/Private/Tests/ToolsetLibraryTest.cpp
new file mode 100644
--- /dev/null
+++ b/Engine/Plugins/Experimental/ToolsetRegistry/Source/ToolsetRegistry/Private/Tests/ToolsetLibraryTest.cpp
@@ -0,0 +1,82 @@
+// Copyright Epic Games, Inc. All Rights Reserved.
+
+#include "CoreMinimal.h"
+#include "JsonObjectConverter.h"
+#include "Misc/AutomationTest.h"
+
+#include "Tests/BlueprintLibraryToolsetTest.h"
+#include "ToolsetLibrary.h"
+#include "ToolsetJsonTest.h"
+#include "ToolsetRegistry/JsonConversion.h"
+
+#if WITH_DEV_AUTOMATION_TESTS
+
+
+BEGIN_DEFINE_SPEC(FToolsetLibrarySpec, "AI.ToolsetRegistry.ToolsetLibrarySpec",
+	EAutomationTestFlags::EditorContext | EAutomationTestFlags::EngineFilter)	
+END_DEFINE_SPEC(FToolsetLibrarySpec)
+
+void FToolsetLibrarySpec::Define()
+{
+	using namespace UE::ToolsetRegistry::Internal;
+
+	It("Generates schema from properties.", [this]()
+	{
+		FString PropertiesJson = UToolsetLibrary::ListClassProperties(
+			UToolsetJsonTestObject::StaticClass());
+		TSharedPtr<FJsonObject> JsonObject = JsonStringToJsonObject(PropertiesJson);
+		TestTrue("Schema", JsonObject.IsValid());
+		if (!JsonObject.IsValid())
+		{
+			return;
+		}
+		TestTrue("Schema", JsonObject->HasField(TEXT("visibleProperty")));
+		TestFalse("Schema", JsonObject->HasField(TEXT("hiddenProperty")));
+	});
+
+	It("Gets values from properties.", [this]()
+	{
+		UToolsetJsonTestObject* TestObject = NewObject<UToolsetJsonTestObject>();
+		TestObject->VisibleProperty = true;
+		FString PropertiesJson = UToolsetLibrary::GetObjectProperties(
+			TestObject, TArray<FName>({ FName("visibleProperty") }));
+		TSharedPtr<FJsonObject> JsonObject = JsonStringToJsonObject(PropertiesJson);
+		TestTrue("Schema", JsonObject.IsValid());
+		if (!JsonObject.IsValid())
+		{
+			return;
+		}
+		bool IsVisible = false;
+		TestTrue("GetProp", JsonObject->TryGetBoolField(TEXT("visibleProperty"), IsVisible));
+		TestTrue("GetProp", IsVisible);
+	});
+
+	It("Sets values on properties.", [this]()
+	{
+		UToolsetJsonTestObject* TestObject = NewObject<UToolsetJsonTestObject>();
+		TestFalse("SetProp", TestObject->VisibleProperty);
+		TSharedPtr<FJsonObject> JsonObject = MakeShared<FJsonObject>();
+		JsonObject->SetBoolField(TEXT("visibleProperty"), true);
+		FString PropertiesJson = JsonToString(JsonObject.ToSharedRef());
+		TestTrue("SetProp", UToolsetLibrary::SetObjectProperties(TestObject, PropertiesJson));
+		TestTrue("SetProp", TestObject->VisibleProperty);
+	});
+
+	It("Returns false when setting values on non-existant properties.", [this]()
+	{
+		UToolsetJsonTestObject* TestObject = NewObject<UToolsetJsonTestObject>();
+		TSharedPtr<FJsonObject> JsonObject = MakeShared<FJsonObject>();
+		JsonObject->SetBoolField(TEXT("doesNotExist"), true);
+		FString PropertiesJson = JsonToString(JsonObject.ToSharedRef());
+		TestFalse("SetProp", UToolsetLibrary::SetObjectProperties(TestObject, PropertiesJson));
+	});
+
+	It("Gets subclasses of a class.", [this]()
+	{
+		TArray<UClass*> Subclasses = UToolsetLibrary::GetDerivedClasses(
+			UToolsetDefinition::StaticClass());
+		TestTrue("Subclass", Subclasses.Contains(UFakeToolset::StaticClass()));
+	});
+}
+
+#endif
diff --git a/Engine/Plugins/Experimental/ToolsetRegistry/Source/ToolsetRegistry/Private/ToolsetLibrary.cpp b/Engine/Plugins/Experimental/ToolsetRegistry/Source/ToolsetRegistry/Private/ToolsetLibrary.cpp
--- a/Engine/Plugins/Experimental/ToolsetRegistry/Source/ToolsetRegistry/Private/ToolsetLibrary.cpp
+++ b/Engine/Plugins/Experimental/ToolsetRegistry/Source/ToolsetRegistry/Private/ToolsetLibrary.cpp
@@ -2,6 +2,8 @@
 
 #include "ToolsetLibrary.h"
 
+#include "AssetRegistry/AssetRegistryHelpers.h"
+#include "AssetRegistry/AssetRegistryModule.h"
 #include "Dom/JsonObject.h"
 #include "LevelEditorViewport.h"
 #include "Misc/AssertionMacros.h"
@@ -10,12 +12,13 @@
 #include "Serialization/JsonReader.h"
 #include "Serialization/JsonWriter.h"
 
+#include "ToolsetRegistryModule.h"
 #include "ToolsetRegistry/JsonConversion.h"
 #include "ToolsetRegistry/ToolsetJson.h"
 
 #include UE_INLINE_GENERATED_CPP_BY_NAME(ToolsetLibrary)
 
-FString UToolsetLibraryTemp::ListClassProperties(UClass* Class)
+FString UToolsetLibrary::ListClassProperties(UClass* Class)
 {
 	if (!Class)
 	{
@@ -34,7 +37,7 @@ FString UToolsetLibraryTemp::ListClassProperties(UClass* Class)
 	return JsonProperties;
 }
 
-FString UToolsetLibraryTemp::GetObjectProperties(const UObject* Object, const TArray<FName>& PropertyNames)
+FString UToolsetLibrary::GetObjectProperties(const UObject* Object, const TArray<FName>& PropertyNames)
 {
 	check(Object);
 	TSharedPtr<FJsonObject> PropertyData = MakeShared<FJsonObject>();
@@ -52,8 +55,13 @@ FString UToolsetLibraryTemp::GetObjectProperties(const UObject* Object, const TA
 	return UE::ToolsetRegistry::Internal::JsonToString(PropertyData.ToSharedRef());
 }
 
-bool UToolsetLibraryTemp::SetObjectProperties(UObject* Object, const FString& PropertiesJson)
+bool UToolsetLibrary::SetObjectProperties(UObject* Object, const FString& PropertiesJson)
 {
+	if (!Object)
+	{
+		return false;
+	}
+
 	bool AllPropertiesSet = false;
 	TSharedPtr<FJsonObject> PropertyData =
 		UE::ToolsetRegistry::Internal::JsonStringToJsonObject(PropertiesJson);
@@ -85,6 +93,41 @@ bool UToolsetLibraryTemp::SetObjectProperties(UObject* Object, const FString& Pr
 	return AllPropertiesSet;
 }
 
+TArray<UClass*> UToolsetLibrary::GetDerivedClasses(UClass* BaseClass)
+{
+	if (!BaseClass)
+	{
+		return TArray<UClass*>();
+	}
+
+	FAssetRegistryModule& AssetRegistryModule = FModuleManager::LoadModuleChecked<FAssetRegistryModule>("AssetRegistry");
+	IAssetRegistry& AssetRegistry = AssetRegistryModule.Get();
+
+	TSet<FTopLevelAssetPath> DerivedNames;
+	AssetRegistry.GetDerivedClassNames({ BaseClass->GetClassPathName() }, {}, DerivedNames);
+
+	TArray<UClass*> DerivedClasses;
+	for (const FTopLevelAssetPath& DerivedName : DerivedNames)
+	{
+		// We ignore skeleton classes (contained in BPs) but consider all other subclasses.
+		if (!DerivedName.GetAssetName().ToString().StartsWith(TEXT("SKEL_")))
+		{
+			FString PathName = DerivedName.ToString();
+			UClass* LoadedClass = FSoftClassPath(PathName).TryLoadClass<UObject>();
+			if (LoadedClass)
+			{
+				DerivedClasses.Add(LoadedClass);
+			}
+			else
+			{
+				UE_LOG(LogToolsetRegistry, Warning,
+					TEXT("Unable to load subclass: %s"), *DerivedName.ToString());
+			}
+		}
+	}
+	return DerivedClasses;
+}
+
 FString UToolsetLibraryTemp::ScreenshotPath(const FString& Filename)
 {
 	FString FullPath = Filename;
diff --git a/Engine/Plugins/Experimental/ToolsetRegistry/Source/ToolsetRegistry/Public/ToolsetLibrary.h b/Engine/Plugins/Experimental/ToolsetRegistry/Source/ToolsetRegistry/Public/ToolsetLibrary.h
--- a/Engine/Plugins/Experimental/ToolsetRegistry/Source/ToolsetRegistry/Public/ToolsetLibrary.h
+++ b/Engine/Plugins/Experimental/ToolsetRegistry/Source/ToolsetRegistry/Public/ToolsetLibrary.h
@@ -8,32 +8,56 @@
 
 #include "ToolsetLibrary.generated.h"
 
+/// Provides functions that are critical to Python toolset operation but which are not available
+/// in and cannot be added to any other BFPL.
 UCLASS(BlueprintType)
-class TOOLSETREGISTRY_API UToolsetLibraryTemp : public UBlueprintFunctionLibrary
+class TOOLSETREGISTRY_API UToolsetLibrary : public UBlueprintFunctionLibrary
 {
 	GENERATED_BODY()
 public:
-
-	// TEMP JSON UTILS
-	// These functions are used by get/set/list properties but we should either move them to JsonUtilities
-	// or promote them to a core part of ToolsetRegistry (if we want to have them use custom JSON exporters).
-
-	/// Returns a JSON Schema formatted string that describes the Blueprint visible properties in the class.
-	UFUNCTION(BlueprintCallable, Category = "Toolset Library Temp")
+	/*
+	 * Returns the Blueprint visible properties in the class as JSON Schema.
+	 * @param Class The class to extract properties from.
+	 * @return A JSON Schema formatted string that describes the properties.
+	 */
+	UFUNCTION(BlueprintCallable, Category = "Toolset Library")
 	static FString ListClassProperties(UClass* Class);
 
-	/// Returns a JSON formatted string that contains the values of the requested properties.
-	UFUNCTION(BlueprintCallable, Category = "Toolset Library Temp")
+	/*
+	 * Returns the values of the requested properties.
+	 * @param Object The object from which to extract properties.
+	 * @param PropertyNames The names of the properties to extract.
+	 * @return A JSON formatted string that contains the values of the requested properties.
+	 */
+	UFUNCTION(BlueprintCallable, Category = "Toolset Library")
 	static FString GetObjectProperties(const UObject* Object, const TArray<FName>& PropertyNames);
 
-	/// Sets UObject properties to the values contained in the Json formatted string.
-	UFUNCTION(BlueprintCallable, Category = "Toolset Library Temp")
+	/*
+	 * Sets the values of specific properties in an Object.
+	 * @param Object The object to modify.
+	 * @param PropertiesJson The property names and values in a JSON formatted string.
+	 * @return True if all properties could be set.
+	 */
+	UFUNCTION(BlueprintCallable, Category = "Toolset Library")
 	static bool SetObjectProperties(UObject* Object, const FString& PropertiesJson);
 
-	// TEMP SCREENSHOT UTILS
-	// These functions are workarounds for the lack of async tool call support. Once we have that,
-	// we should remove these functions and implement a tool that directly takes a screenshot and
-	// returns a ToolsetImage.
+	/*
+	 * Returns the list of subclasses that derive from a class.
+	 * @param BaseClass The class to get subclasses from.
+	 * @return The list of native and BP classes that derive from the base class.
+	 */
+	UFUNCTION(BlueprintCallable, Category = "Toolset Library")
+	static TArray<UClass*> GetDerivedClasses(UClass* BaseClass);
+};
+
+// Provides workarounds for the lack of async tool call support. Once we have that,
+// we should remove these functions and implement a tool that directly takes a screenshot and
+// returns a ToolsetImage.
+UCLASS(BlueprintType)
+class TOOLSETREGISTRY_API UToolsetLibraryTemp : public UBlueprintFunctionLibrary
+{
+	GENERATED_BODY()
+public:
 
 	// Converts the file name of a screenshot into a full path based on UEs default logic.
 	UFUNCTION(BlueprintCallable, Category = "Toolset Library Temp")
@@ -43,3 +67,4 @@ public:
 	UFUNCTION(BlueprintCallable, Category = "Toolset Library Temp")
 	static FToolsetImage ScreenshotToImage(const FString& FullPath);
 };
+
