Hide test and internal toolsets from UnrealEd UIs so that they don't clutter places like the dropdown in AgentSkill.

#Synced-CL 51002830

diff --git a/Engine/Plugins/Experimental/AIAssistant/Source/AIAssistant/Private/AIAssistantToolset.h b/Engine/Plugins/Experimental/AIAssistant/Source/AIAssistant/Private/AIAssistantToolset.h
--- a/Engine/Plugins/Experimental/AIAssistant/Source/AIAssistant/Private/AIAssistantToolset.h
+++ b/Engine/Plugins/Experimental/AIAssistant/Source/AIAssistant/Private/AIAssistantToolset.h
@@ -74,7 +74,7 @@ public:
 };
 
 /// These functions help the AI assistant inspect its state as it relates to the larger Unreal Editor.
-UCLASS(BlueprintType)
+UCLASS(BlueprintType, Hidden)
 class UAIAssistantToolset : public UToolsetDefinition
 {
 	GENERATED_BODY()
diff --git a/Engine/Plugins/Experimental/ToolsetRegistry/Source/ToolsetRegistry/Private/Tests/BlueprintLibraryToolsetTest.h b/Engine/Plugins/Experimental/ToolsetRegistry/Source/ToolsetRegistry/Private/Tests/BlueprintLibraryToolsetTest.h
--- a/Engine/Plugins/Experimental/ToolsetRegistry/Source/ToolsetRegistry/Private/Tests/BlueprintLibraryToolsetTest.h
+++ b/Engine/Plugins/Experimental/ToolsetRegistry/Source/ToolsetRegistry/Private/Tests/BlueprintLibraryToolsetTest.h
@@ -12,7 +12,7 @@
 #include "BlueprintLibraryToolsetTest.generated.h"
 
 /** Fake Toolset that does nothing. */
-UCLASS(Blueprintable)
+UCLASS(Blueprintable, Hidden)
 class UFakeToolset : public UToolsetDefinition
 {
 	GENERATED_BODY()
@@ -41,7 +41,7 @@ public:
 static TSharedPtr<FJsonObject> GetFakeToolsetExpectedSchema();
 
 /** Fake Toolset that has method with default parameters. */
-UCLASS(Blueprintable)
+UCLASS(Blueprintable, Hidden)
 class UFakeToolsetWithDefaultParams : public UToolsetDefinition
 {
 	GENERATED_BODY()
@@ -60,7 +60,7 @@ static TSharedPtr<FJsonObject> GetFakeToolsetWithDefaultParamsExpectedSchema();
 
 
 /** Fake Toolset with no tools. */
-UCLASS(Blueprintable)
+UCLASS(Blueprintable, Hidden)
 class UFakeToolsetWithNoTools : public UToolsetDefinition
 {
 	GENERATED_BODY()
@@ -69,7 +69,7 @@ class UFakeToolsetWithNoTools : public UToolsetDefinition
 static TSharedPtr<FJsonObject> GetFakeToolsetWithNoToolsExpectedSchema();
 
 /** Fake Toolset that has method with TArray parameters. */
-UCLASS(Blueprintable)
+UCLASS(Blueprintable, Hidden)
 class UFakeToolsetWithArrayParams : public UToolsetDefinition
 {
 	GENERATED_BODY()
@@ -87,7 +87,7 @@ public:
 static TSharedPtr<FJsonObject> GetFakeToolsetWithArrayParamsExpectedSchema();
 
 /** Fake Toolset that has method with TMap parameters. */
-UCLASS(Blueprintable)
+UCLASS(Blueprintable, Hidden)
 class UFakeToolsetWithMapParams : public UToolsetDefinition
 {
 	GENERATED_BODY()
@@ -105,7 +105,7 @@ public:
 static TSharedPtr<FJsonObject> GetFakeToolsetWithMapParamsExpectedSchema();
 
 /** Fake Toolset that has method with TSet parameters. */
-UCLASS(Blueprintable)
+UCLASS(Blueprintable, Hidden)
 class UFakeToolsetWithSetParams : public UToolsetDefinition
 {
 	GENERATED_BODY()
@@ -123,7 +123,7 @@ public:
 static TSharedPtr<FJsonObject> GetFakeToolsetWithSetParamsExpectedSchema();
 
 /** Fake Toolset with private and non-static methods. */
-UCLASS(Blueprintable)
+UCLASS(Blueprintable, Hidden)
 class UFakeToolsetWithPrivateMethods : public UToolsetDefinition
 {
 	GENERATED_BODY()
@@ -153,7 +153,7 @@ private:
 };
 
 /** Fake Toolset that has error-prone method. */
-UCLASS(Blueprintable)
+UCLASS(Blueprintable, Hidden)
 class UFakeErrorProneToolset : public UToolsetDefinition
 {
 	GENERATED_BODY()
diff --git a/Engine/Plugins/Experimental/ToolsetRegistry/Source/ToolsetRegistry/Public/AgentSkill.h b/Engine/Plugins/Experimental/ToolsetRegistry/Source/ToolsetRegistry/Public/AgentSkill.h
--- a/Engine/Plugins/Experimental/ToolsetRegistry/Source/ToolsetRegistry/Public/AgentSkill.h
+++ b/Engine/Plugins/Experimental/ToolsetRegistry/Source/ToolsetRegistry/Public/AgentSkill.h
@@ -59,7 +59,7 @@ protected:
 
 
 
-UCLASS(BlueprintType)
+UCLASS(BlueprintType, Hidden)
 class TOOLSETREGISTRY_API UAgentSkillToolset : public UToolsetDefinition
 {
 	GENERATED_BODY()
diff --git a/Engine/Plugins/Experimental/ToolsetRegistry/Source/ToolsetRegistry/Public/BPToolsetRegistry.h b/Engine/Plugins/Experimental/ToolsetRegistry/Source/ToolsetRegistry/Public/BPToolsetRegistry.h
--- a/Engine/Plugins/Experimental/ToolsetRegistry/Source/ToolsetRegistry/Public/BPToolsetRegistry.h
+++ b/Engine/Plugins/Experimental/ToolsetRegistry/Source/ToolsetRegistry/Public/BPToolsetRegistry.h
@@ -10,7 +10,7 @@
 
 // This is the common base class for Toolsets defined as UObjects.
 // UFunctions that define tools in this class should be static functions.
-UCLASS(BlueprintType)
+UCLASS(BlueprintType, Abstract)
 class TOOLSETREGISTRY_API UToolsetDefinition : public UObject
 {
 	GENERATED_BODY()
