Convert FColor into FLinearColor in ToolsetJson so that there is a single LLM-facing color representation.

#Synced-CL 51002830

diff --git a/Engine/Plugins/Experimental/AIAssistant/Source/AIAssistant/Private/AIAssistantToolset.cpp b/Engine/Plugins/Experimental/AIAssistant/Source/AIAssistant/Private/AIAssistantToolset.cpp
--- a/Engine/Plugins/Experimental/AIAssistant/Source/AIAssistant/Private/AIAssistantToolset.cpp
+++ b/Engine/Plugins/Experimental/AIAssistant/Source/AIAssistant/Private/AIAssistantToolset.cpp
@@ -31,7 +31,7 @@ FAIAssistantContext UAIAssistantToolset::GetProjectContext()
 		TEXT("100 units is 1 meter"),
 		TEXT("Screen space coordinates are normalized coordinates where 0, 0 is the upper left corner "
 			 "and 1, 1 is the lower right corner."),
-		TEXT("Colors are [0,255]. LinearColors are [0,1]"),
+		TEXT("Colors with a max of 1 are LDR. No max means HDR."),
 		TEXT("If you aren't sure what the user is referring to, try checked to see what they have selected.")
 	});
 
diff --git a/Engine/Plugins/Experimental/ToolsetRegistry/Source/ToolsetRegistry/Private/Tests/ToolsetJsonTest.cpp b/Engine/Plugins/Experimental/ToolsetRegistry/Source/ToolsetRegistry/Private/Tests/ToolsetJsonTest.cpp
--- a/Engine/Plugins/Experimental/ToolsetRegistry/Source/ToolsetRegistry/Private/Tests/ToolsetJsonTest.cpp
+++ b/Engine/Plugins/Experimental/ToolsetRegistry/Source/ToolsetRegistry/Private/Tests/ToolsetJsonTest.cpp
@@ -32,7 +32,7 @@ namespace
 			Test.TestTrue("Title", Schema->TryGetStringField(TEXT("title"), Title));
 			FString ClassPrefix(IsClass ? TEXT("Class@") : TEXT(""));
 			Test.TestEqual("Title", Title, ClassPrefix + Class->GetPathName());
-		}	
+		}
 	
 		const TSharedPtr<FJsonObject>* Properties = nullptr;
 		Test.TestTrue("Properties", Schema->TryGetObjectField(TEXT("properties"), Properties));
@@ -80,6 +80,97 @@ namespace
 			TestObjectRefData(Test, *Property, Expected);
 		}
 	}
+
+	void TestColorSchema(
+		FAutomationTestBase& Test, const TSharedPtr<FJsonObject>& Schema, bool HasMax)
+	{
+		Test.TestTrue("Schema", Schema.IsValid());
+		if (!Schema.IsValid())
+		{
+			return;
+		}
+
+		FString Type;
+		Test.TestTrue("Type", Schema->TryGetStringField(TEXT("type"), Type));
+		Test.TestEqual("Type", Type, FString(TEXT("object")));
+
+		FString Title;
+		Test.TestTrue("Title", Schema->TryGetStringField(TEXT("title"), Title));
+		Test.TestEqual("Title", Title, FString(TEXT("LinearColor")));
+
+		Test.TestFalse("Description", Schema->HasField(TEXT("description")));
+
+		const TSharedPtr<FJsonObject>* Properties = nullptr;
+		Test.TestTrue("Properties", Schema->TryGetObjectField(TEXT("properties"), Properties));
+		if (Properties)
+		{
+			TArray<FString> PropertyNames({ "r", "g", "b", "a" });
+			for (const auto& Name : PropertyNames)
+			{
+				const TSharedPtr<FJsonObject>* Property = nullptr;
+				Test.TestTrue(
+					"Property", Properties->Get()->TryGetObjectField(Name, Property));
+				if (Property)
+				{
+					Test.TestEqual("Property", Property->Get()->GetStringField(
+						TEXT("type")), FString(TEXT("number")));
+					Test.TestEqual("Property", Property->Get()->GetNumberField(
+						TEXT("minimum")), 0.0);
+					if (HasMax)
+					{
+						Test.TestEqual("Property", Property->Get()->GetNumberField(
+							TEXT("maximum")), 1.0);
+					}
+				}
+			}
+		}
+	}
+
+	void TestColorData(
+		FAutomationTestBase& Test, const TSharedPtr<FJsonObject>& Schema, const FColor& Expected)
+	{
+		Test.TestTrue("Schema", Schema.IsValid());
+		if (!Schema.IsValid())
+		{
+			return;
+		}
+
+		const FLinearColor LinearExpected = Expected.ReinterpretAsLinear();
+
+		float R;
+		Test.TestTrue("Color", Schema->TryGetNumberField(TEXT("r"), R));
+		Test.TestEqual("Color", R, LinearExpected.R, 2e-3f);
+
+		float G;
+		Test.TestTrue("Color", Schema->TryGetNumberField(TEXT("g"), G));
+		Test.TestEqual("Color", G, LinearExpected.G, 2e-3f);
+
+		float B;
+		Test.TestTrue("Color", Schema->TryGetNumberField(TEXT("b"), B));
+		Test.TestEqual("Color", B, LinearExpected.B, 2e-3f);
+
+		float A;
+		Test.TestTrue("Color", Schema->TryGetNumberField(TEXT("a"), A));
+		Test.TestEqual("Color", A, LinearExpected.A, 2e-3f);
+	}
+
+	void TestColorDataInStruct(
+		FAutomationTestBase& Test, const TSharedPtr<FJsonObject>& Schema,
+		const FString& PropertyName, const FColor& Expected)
+	{
+		Test.TestTrue("Schema", Schema.IsValid());
+		if (!Schema.IsValid())
+		{
+			return;
+		}
+
+		const TSharedPtr<FJsonObject>* Property = nullptr;
+		Test.TestTrue("Unpack", Schema->TryGetObjectField(PropertyName, Property));
+		if (Property)
+		{
+			TestColorData(Test, *Property, Expected);
+		}
+	}
 }
 
 BEGIN_DEFINE_SPEC(FToolsetJsonSpec, "AI.ToolsetRegistry.ToolsetJsonSpec",
@@ -238,6 +329,54 @@ void FToolsetJsonSpec::Define()
 		TestEqual("RefCheck", TestIn.TestClassPath, TestOut.TestClassPath);
 	});
 
+	It("Converts FColor property to FLinearColor schema", [this]()
+	{
+		UStruct* Struct = FToolsetJsonTest::StaticStruct();
+		FProperty* Prop = Struct->FindPropertyByName("TestColor");
+		const TSharedPtr<FJsonObject> Schema = PropertyToJsonSchema(Prop);
+		TestColorSchema(*this, Schema, true);
+	});
+
+	It("Converts FColor property to and from JSON", [this]()
+	{
+		FToolsetJsonTest TestIn;
+		TestIn.TestColor = FColor(10, 20, 30, 40);
+
+		UStruct* Struct = FToolsetJsonTest::StaticStruct();
+		FProperty* Property = Struct->FindPropertyByName("TestColor");
+
+		TSharedPtr<FJsonValue> OutJson = PropertyToJsonData(Property, &TestIn.TestColor);
+		TestColorData(*this, OutJson->AsObject(), TestIn.TestColor);
+
+		FToolsetJsonTest TestOut;
+		TestTrue("Output", JsonDataToProperty(OutJson, Property, &TestOut.TestColor));
+		TestEqual("ColorCheck", TestIn.TestColor, TestOut.TestColor);
+	});
+
+	It("Converts FLinearColor property to FLinearColor schema", [this]()
+	{
+		UStruct* Struct = FToolsetJsonTest::StaticStruct();
+		FProperty* Prop = Struct->FindPropertyByName("TestLinearColor");
+		const TSharedPtr<FJsonObject> Schema = PropertyToJsonSchema(Prop);
+		TestColorSchema(*this, Schema, false);
+	});
+
+	It("Converts FLinearColor property to and from JSON", [this]()
+	{
+		FToolsetJsonTest TestIn;
+		TestIn.TestLinearColor = FLinearColor(0.1, 0.2, 0.3, 0.4);
+
+		UStruct* Struct = FToolsetJsonTest::StaticStruct();
+		FProperty* Property = Struct->FindPropertyByName("TestLinearColor");
+
+		TSharedPtr<FJsonValue> OutJson = PropertyToJsonData(Property, &TestIn.TestLinearColor);
+		TestColorData(*this, OutJson->AsObject(), TestIn.TestLinearColor.ToFColor(false));
+
+		FToolsetJsonTest TestOut;
+		TestTrue("Output", JsonDataToProperty(OutJson, Property, &TestOut.TestLinearColor));
+		TestEqual("ColorCheck", TestIn.TestLinearColor, TestOut.TestLinearColor);
+	});
+
 	It("Converts struct to reference schema", [this]()
 	{		
 		TSharedPtr<FJsonObject> Schema = StructToJsonSchema(FToolsetJsonTest::StaticStruct());
@@ -258,7 +397,15 @@ void FToolsetJsonSpec::Define()
 				TestTrue("Object", Properties->Get()->TryGetObjectField(JsonPropertyName, TestProperty));
 				if (TestProperty)
 				{
-					TestObjectRefSchema(*this, *TestProperty, nullptr, false);
+					if (JsonPropertyName == FString(TEXT("testColor")) ||
+						JsonPropertyName == FString(TEXT("testLinearColor")))
+					{
+						TestColorSchema(*this, *TestProperty, JsonPropertyName == FString(TEXT("testColor")));
+					}
+					else
+					{
+						TestObjectRefSchema(*this, *TestProperty, nullptr, false);
+					}
 				}
 			}
 		}
@@ -274,6 +421,7 @@ void FToolsetJsonSpec::Define()
 		TestIn.TestObjectWeak = TestOther;
 		TestIn.TestClass = UToolsetJsonTestObject::StaticClass();
 		TestIn.TestClassPath = UToolsetJsonTestObject::StaticClass();
+		TestIn.TestColor = FColor(10, 20, 30, 40);
 
 		TSharedPtr<FJsonObject> JsonObject = StructToJsonData(
 			FToolsetJsonTest::StaticStruct(), &TestIn);
@@ -286,6 +434,16 @@ void FToolsetJsonSpec::Define()
 				TestObjectRefDataInStruct(
 					*this, JsonObject, JsonPropertyName, UToolsetJsonTestObject::StaticClass()->GetPathName());
 			}
+			else if (JsonPropertyName == "testColor")
+			{
+				TestColorDataInStruct(
+					*this, JsonObject, JsonPropertyName, TestIn.TestColor);
+			}
+			else if (JsonPropertyName == "testLinearColor")
+			{
+				TestColorDataInStruct(
+					*this, JsonObject, JsonPropertyName, TestIn.TestLinearColor.ToFColor(false));
+			}
 			else
 			{
 				TestObjectRefDataInStruct(
@@ -302,6 +460,7 @@ void FToolsetJsonSpec::Define()
 		TestEqual("RefCheck", TestIn.TestObjectWeak, TestOut.TestObjectWeak);
 		TestEqual("RefCheck", TestIn.TestClass, TestOut.TestClass);
 		TestEqual("RefCheck", TestIn.TestClassPath, TestOut.TestClassPath);
+		TestEqual("ColorCheck", TestIn.TestColor, TestOut.TestColor);
 	});
 
 	It("Imports from old string style path", [this]()
@@ -369,6 +528,23 @@ void FToolsetJsonSpec::Define()
 				{
 					TestTrue("Default", TestObject->Get()->GetField(TEXT("default"), EJson::Null)->IsNull());
 				}
+
+				const TSharedPtr<FJsonObject>* TestColor = nullptr;
+				TestTrue("Input", Properties->Get()->TryGetObjectField(TEXT("testColor"), TestColor));
+				if (TestColor)
+				{
+					TestColorData(
+						*this, TestColor->Get()->GetObjectField(TEXT("default")), FColor(10, 20, 30, 40));
+				}
+
+				const TSharedPtr<FJsonObject>* TestLinearColor = nullptr;
+				TestTrue("Input", Properties->Get()->TryGetObjectField(TEXT("testLinearColor"), TestLinearColor));
+				if (TestLinearColor)
+				{
+					TestColorData(
+						*this, TestLinearColor->Get()->GetObjectField(TEXT("default")),
+						FLinearColor(0.1, 0.2, 0.3, 0.4).ToFColor(false));
+				}
 			}
 		}
 	});
diff --git a/Engine/Plugins/Experimental/ToolsetRegistry/Source/ToolsetRegistry/Private/Tests/ToolsetJsonTest.h b/Engine/Plugins/Experimental/ToolsetRegistry/Source/ToolsetRegistry/Private/Tests/ToolsetJsonTest.h
--- a/Engine/Plugins/Experimental/ToolsetRegistry/Source/ToolsetRegistry/Private/Tests/ToolsetJsonTest.h
+++ b/Engine/Plugins/Experimental/ToolsetRegistry/Source/ToolsetRegistry/Private/Tests/ToolsetJsonTest.h
@@ -21,7 +21,9 @@ public:
 	bool HiddenProperty;
 
 	UFUNCTION(BlueprintCallable, Category = "Main")
-	void TestDefault(UObject* TestObject = nullptr) {}
+	void TestDefault(
+		UObject* TestObject = nullptr, FColor TestColor = FColor(10, 20, 30, 40),
+		FLinearColor TestLinearColor = FLinearColor(0.1, 0.2, 0.3, 0.4)) {}
 };
 
 USTRUCT(BlueprintType)
@@ -46,4 +48,10 @@ public:
 
 	UPROPERTY()
 	FSoftClassPath TestClassPath;
+
+	UPROPERTY()
+	FColor TestColor = FColor::White;
+
+	UPROPERTY()
+	FLinearColor TestLinearColor = FLinearColor::White;
 };
diff --git a/Engine/Plugins/Experimental/ToolsetRegistry/Source/ToolsetRegistry/Private/ToolsetRegistry/ToolsetJson.cpp b/Engine/Plugins/Experimental/ToolsetRegistry/Source/ToolsetRegistry/Private/ToolsetRegistry/ToolsetJson.cpp
--- a/Engine/Plugins/Experimental/ToolsetRegistry/Source/ToolsetRegistry/Private/ToolsetRegistry/ToolsetJson.cpp
+++ b/Engine/Plugins/Experimental/ToolsetRegistry/Source/ToolsetRegistry/Private/ToolsetRegistry/ToolsetJson.cpp
@@ -69,11 +69,13 @@ namespace
 		return false;
 	}
 
-	const FStructProperty* AsColorProperty(TNotNull<const FProperty*> Property)
+	const FStructProperty* AsColorProperty(
+		TNotNull<const FProperty*> Property, bool AllowLinearColor = false)
 	{
 		if (const FStructProperty* StructProperty = CastField<FStructProperty>(Property))
 		{
-			if (StructProperty->Struct->GetFName() == NAME_Color)
+			const FName StructName = StructProperty->Struct->GetFName();
+			if (StructName == NAME_Color || (AllowLinearColor && StructName == NAME_LinearColor))
 			{
 				return StructProperty;
 			}
@@ -144,6 +146,42 @@ bool CustomJsonSchema(
 		}
 		return true;
 	}
+	else if (const FStructProperty* ColorProperty = AsColorProperty(Property, true))
+	{
+		bool IsLinear = ColorProperty->Struct->GetFName() == NAME_LinearColor;
+		if (ParameterDefaultString.IsEmpty())
+		{
+			TSharedPtr Schema = FJsonSchemaGenerator::UStructToJsonSchemaObject(
+				TBaseStructure<FLinearColor>::Get());
+			TSharedPtr<FJsonObject> Properties = Schema->GetObjectField(TEXT("properties"));
+			TArray<FString> PropertyNames({ "r", "g", "b", "a" });
+			for (const auto& Name : PropertyNames)
+			{
+				TSharedPtr<FJsonObject> Component = Properties->GetObjectField(Name);
+				Component->SetNumberField(FString(TEXT("minimum")), 0.f);
+				if (!IsLinear)
+				{
+					Component->SetNumberField(FString(TEXT("maximum")), 1.f);
+				}
+			}
+			// The default description isn't very useful and includes a source code reference,
+			// so we just remove it to save some tokens.
+			Schema->RemoveField(FString(TEXT("description")));
+			FJsonObject::Duplicate(Schema, OutputSchema);
+			return true;
+		}
+		else if (!IsLinear)
+		{
+			FColor DefaultColor = FColor::Black;
+			Property->ImportText_Direct(*ParameterDefaultString, &DefaultColor, nullptr, PPF_None);
+			FLinearColor DefaultLinearColor = DefaultColor.ReinterpretAsLinear();
+			TSharedPtr<FJsonObject> OutJson = MakeShared<FJsonObject>();
+			FJsonObjectConverter::UStructToJsonObject(
+				TBaseStructure<FLinearColor>::Get(), &DefaultLinearColor, OutJson.ToSharedRef());
+			OutputSchema->SetObjectField(TEXT("default"), OutJson);
+			return true;
+		}
+	}
 	return false;
 }
 
@@ -169,6 +207,21 @@ bool CustomImport(const TSharedPtr<FJsonValue>& JsonValue, FProperty* Property,
 			return true;
 		}
 	}
+	else if (AsColorProperty(Property))
+	{
+		TSharedPtr<FJsonObject> JsonObject = JsonValue->AsObject();
+		if (JsonObject.IsValid())
+		{
+			FLinearColor LinearColor;
+			if (FJsonObjectConverter::JsonObjectToUStruct(
+				JsonObject.ToSharedRef(), TBaseStructure<FLinearColor>::Get(), &LinearColor))
+			{
+				FColor& ColorOut = *(static_cast<FColor*>(Value));
+				ColorOut = LinearColor.ToFColor(false);
+				return true;
+			}
+		}
+	}
 	return false;
 }
 
@@ -215,6 +268,15 @@ TSharedPtr<FJsonValue> CustomExport(FProperty* Property, const void* Value)
 		}
 		return MakeShared<FJsonValueObject>(FJsonObjectConverter::UStructToJsonObject(Reference));
 	}
+	else if (AsColorProperty(Property))
+	{
+		const FColor& InColor = *(static_cast<const FColor*>(Value));
+		FLinearColor OutColor = InColor.ReinterpretAsLinear();
+		TSharedPtr<FJsonObject> OutJson = MakeShared<FJsonObject>();
+		FJsonObjectConverter::UStructToJsonObject(
+			TBaseStructure<FLinearColor>::Get(), &OutColor, OutJson.ToSharedRef());
+		return MakeShared<FJsonValueObject>(OutJson);
+	}
 	
 	return nullptr;
 }
