[AI Assistant] Accessor to get the JSON schema for the value of async results.

ToolsetJson will delegate to UToolsetAsyncResult to get the schema for the
value type rather than polluting ToolsetJson.cpp with the implementation details.

#Synced-CL 51097503

diff --git a/Engine/Plugins/Experimental/ToolsetRegistry/Source/ToolsetRegistry/Private/Tests/ToolCallAsyncResultTest.cpp b/Engine/Plugins/Experimental/ToolsetRegistry/Source/ToolsetRegistry/Private/Tests/ToolCallAsyncResultTest.cpp
--- a/Engine/Plugins/Experimental/ToolsetRegistry/Source/ToolsetRegistry/Private/Tests/ToolCallAsyncResultTest.cpp
+++ b/Engine/Plugins/Experimental/ToolsetRegistry/Source/ToolsetRegistry/Private/Tests/ToolCallAsyncResultTest.cpp
@@ -10,12 +10,15 @@
 #include "Templates/FunctionFwd.h"
 #include "Templates/SharedPointer.h"
 #include "Templates/UnrealTemplate.h"
+#include "UObject/Field.h"
 #include "UObject/Object.h"
 #include "UObject/ObjectPtr.h"
+#include "UObject/UnrealType.h"
 
 #include "ToolsetRegistry/JsonConversion.h"
 #include "ToolsetRegistry/JsonValueOrError.h"
 #include "ToolsetRegistry/ToolCallAsyncResult.h"
+#include "ToolsetRegistry/ToolsetJson.h"
 
 #if WITH_DEV_AUTOMATION_TESTS
 
@@ -455,6 +458,24 @@ void FToolCallAsyncResultBaseSpec::Define()
 					}),
 				TEXT("Schedule broadcast on completed")));
 	});
+
+	It(TEXT("Should provide the value's JSON schema"), [this]()
+	{
+		using namespace UE::ToolsetRegistry::Internal;
+
+		// This is a bit odd since it partially mirrors the implementation. However, the
+		// alternative would be to force all value types to specify their expected schema.
+		FProperty* Property = Result->GetClass()->FindPropertyByName(TEXT("Value"));
+		if (!Property) return;  // No value, ignore
+
+		TSharedPtr<FJsonObject> ExpectedSchema = ToolsetJson::PropertyToJsonSchema(Property);
+		if (!TestTrue(TEXT("Value has schema"), ExpectedSchema.IsValid())) return;
+
+		TestEqual(
+			TEXT("Value schema"),
+			JsonToString(UToolCallAsyncResult::GetValueJsonSchema(Result->GetClass())),
+			JsonToString(ExpectedSchema.ToSharedRef()));
+	});
 }
 
 // Test functionality of UToolCallAsyncResult where there is no way to successfully complete.
@@ -482,4 +503,40 @@ namespace
 		TEXT("FToolCallAsyncResultNoCompletionSpec"));
 }
 
+BEGIN_DEFINE_SPEC(
+	FToolCallAsyncResultTestProperty,
+	"AI.ToolsetRegistry.ToolCallAsyncResult.Property",
+	EAutomationTestFlags::EditorContext | EAutomationTestFlags::EngineFilter)
+END_DEFINE_SPEC(FToolCallAsyncResultTestProperty)
+
+void FToolCallAsyncResultTestProperty::Define()
+{
+	Describe(TEXT("MatchesProperty"), [this]()
+	{
+		It(TEXT("Should match an async result class"), [this]()
+		{
+			UFunction* Function = UToolCallAsyncResultReturnPropertyTest::FindTestFunction();
+			if (!TestTrue(TEXT("Has test function"), Function != nullptr)) return;
+			FProperty* ReturnProperty = Function->GetReturnProperty();
+			if (!TestTrue(TEXT("Has return property"), ReturnProperty != nullptr)) return;
+			UClass* ReturnClass = UToolCallAsyncResult::MatchesProperty(ReturnProperty);
+			if (!TestTrue(TEXT("Is a class"), ReturnClass != nullptr)) return;
+			TestTrue(
+				TEXT("Returns async result class"),
+				ReturnClass->IsChildOf<UToolCallAsyncResult>());
+		});
+
+		It(TEXT("Should not match a non-async result class"), [this]()
+		{
+			UFunction* Function = UToolCallAsyncResultReturnPropertyTest::FindTestFunction();
+			if (!TestTrue(TEXT("Has test function"), Function != nullptr)) return;
+			FProperty* Argument = Function->FindPropertyByName(TEXT("Subject"));
+			if (!TestTrue(TEXT("Has argument"), Argument != nullptr)) return;
+			TestTrue(
+				TEXT("Does not match non-async result property"),
+				UToolCallAsyncResult::MatchesProperty(Argument) == nullptr);
+		});
+	});
+}
+
 #endif  // WITH_DEV_AUTOMATION_TESTS
\ No newline at end of file
diff --git a/Engine/Plugins/Experimental/ToolsetRegistry/Source/ToolsetRegistry/Private/Tests/ToolCallAsyncResultTest.h b/Engine/Plugins/Experimental/ToolsetRegistry/Source/ToolsetRegistry/Private/Tests/ToolCallAsyncResultTest.h
--- a/Engine/Plugins/Experimental/ToolsetRegistry/Source/ToolsetRegistry/Private/Tests/ToolCallAsyncResultTest.h
+++ b/Engine/Plugins/Experimental/ToolsetRegistry/Source/ToolsetRegistry/Private/Tests/ToolCallAsyncResultTest.h
@@ -11,6 +11,9 @@
 
 #include "ToolsetRegistry/ToolCallAsyncResult.h"
 #include "ToolsetRegistry/ToolCallAsyncResultFutureHandler.h"
+#include "ToolsetRegistry/ToolCallAsyncResultString.h"
+
+#include "ToolCallAsyncResultTest.generated.h"
 
 BEGIN_DEFINE_SPEC(
 	FToolCallAsyncResultBaseSpec,
@@ -103,4 +106,27 @@ BEGIN_DEFINE_SPEC(
 		    FToolCallAsyncResult##AsyncResultValueShortName##SpecInstance(                         \
 		        TEXT(UE_INLINE_STRINGIFY(                                                          \
 		            FToolCallAsyncResultWith##AsyncResultValueShortName##Spec)));                  \
-	}
\ No newline at end of file
+	}
+
+// Used to test matching UFunction that returns UToolCallAsyncResult.
+UCLASS(BlueprintType, Hidden)
+class UToolCallAsyncResultReturnPropertyTest : public UObject
+{
+	GENERATED_BODY()
+
+public:
+	UFUNCTION()
+	static UToolCallAsyncResultString* SlowGreeting(
+		const FString& Subject)
+	{
+		TObjectPtr<UToolCallAsyncResultString> Result =
+			NewObject<UToolCallAsyncResultString>();
+		Result->SetValue(FString::Printf(TEXT("Hello %s"), *Subject));
+		return Result.Get();
+	}
+
+	static UFunction* FindTestFunction()
+	{
+		return StaticClass()->FindFunctionByName(TEXT("SlowGreeting"));
+	}
+};
\ No newline at end of file
diff --git a/Engine/Plugins/Experimental/ToolsetRegistry/Source/ToolsetRegistry/Private/Tests/ToolCallAsyncResultVoidTest.cpp b/Engine/Plugins/Experimental/ToolsetRegistry/Source/ToolsetRegistry/Private/Tests/ToolCallAsyncResultVoidTest.cpp
--- a/Engine/Plugins/Experimental/ToolsetRegistry/Source/ToolsetRegistry/Private/Tests/ToolCallAsyncResultVoidTest.cpp
+++ b/Engine/Plugins/Experimental/ToolsetRegistry/Source/ToolsetRegistry/Private/Tests/ToolCallAsyncResultVoidTest.cpp
@@ -7,6 +7,7 @@
 #include "UObject/ObjectPtr.h"
 
 #include "ToolsetRegistry/ToolCallAsyncResultVoid.h"
+#include "ToolsetRegistry/JsonConversion.h"
 #include "Tests/ToolCallAsyncResultTest.h"
 
 #if WITH_DEV_AUTOMATION_TESTS
@@ -26,6 +27,15 @@ public:
 			};
 		ExpectedJsonValue = MakeShared<FJsonValueNull>();
 		BaseSpec::Define();
+
+		It(TEXT("Should have a null JSON schema for the value"), [this]()
+		{
+			TestEqual(
+				TEXT("JSON schema"),
+				UE::ToolsetRegistry::Internal::JsonToString(
+					UToolCallAsyncResultVoid::GetValueJsonSchema()),
+				TEXT(R"json({"description":"Always null","type":"null"})json"));
+		});
 	}
 
 	virtual FString GetBeautifiedTestName() const override
diff --git a/Engine/Plugins/Experimental/ToolsetRegistry/Source/ToolsetRegistry/Private/ToolsetRegistry/ToolCallAsyncResult.cpp b/Engine/Plugins/Experimental/ToolsetRegistry/Source/ToolsetRegistry/Private/ToolsetRegistry/ToolCallAsyncResult.cpp
--- a/Engine/Plugins/Experimental/ToolsetRegistry/Source/ToolsetRegistry/Private/ToolsetRegistry/ToolCallAsyncResult.cpp
+++ b/Engine/Plugins/Experimental/ToolsetRegistry/Source/ToolsetRegistry/Private/ToolsetRegistry/ToolCallAsyncResult.cpp
@@ -13,15 +13,18 @@
 
 #include "ToolsetRegistryModule.h"
 #include "ToolsetRegistry/JsonConversion.h"
+#include "ToolsetRegistry/ToolCallAsyncResultVoid.h"
 #include "ToolsetRegistry/ToolsetJson.h"
 
+FString UToolCallAsyncResult::ValuePropertyName = TEXT("Value");
+
 TSharedPtr<FJsonValue> UToolCallAsyncResult::GetValueAsJson() const
 {
 	if (!bIsComplete || !Error.IsEmpty()) return TSharedPtr<FJsonValue>();
 
 	UClass* Class = GetClass();
 	check(Class);
-	FProperty* Property = Class->FindPropertyByName(TEXT("Value"));
+	FProperty* Property = Class->FindPropertyByName(*ValuePropertyName);
 	if (!Property)
 	{
 		UE_LOG(
@@ -130,4 +133,30 @@ void UToolCallAsyncResult::RunOnMainThread(TUniqueFunction<void()>&& Function)
 	{
 		AsyncTask(ENamedThreads::GameThread, MoveTemp(Function));
 	}
+}
+
+TSharedRef<FJsonObject> UToolCallAsyncResult::GetValueJsonSchema(
+	TSubclassOf<UToolCallAsyncResult> Class)
+{
+	using namespace UE::ToolsetRegistry::Internal;
+	check(Class);
+	FProperty* Property = Class->FindPropertyByName(*ValuePropertyName);
+	TSharedPtr<FJsonObject> Schema =
+		Property
+		? ToolsetJson::PropertyToJsonSchema(Property)
+		: UToolCallAsyncResultVoid::GetValueJsonSchema();
+	check(Schema.IsValid());
+	return Schema.ToSharedRef();
+}
+
+UClass* UToolCallAsyncResult::MatchesProperty(TNotNull<const FProperty*> Property)
+{
+	TObjectPtr<UClass> Class;
+	const FObjectProperty* ObjectProperty = CastField<FObjectProperty>(Property);
+	if (ObjectProperty)
+	{
+		Class = ObjectProperty->PropertyClass;
+		if (!Class->IsChildOf<UToolCallAsyncResult>()) Class = nullptr;
+	}
+	return Class;
 }
\ No newline at end of file
diff --git a/Engine/Plugins/Experimental/ToolsetRegistry/Source/ToolsetRegistry/Private/ToolsetRegistry/ToolCallAsyncResultVoid.cpp b/Engine/Plugins/Experimental/ToolsetRegistry/Source/ToolsetRegistry/Private/ToolsetRegistry/ToolCallAsyncResultVoid.cpp
--- a/Engine/Plugins/Experimental/ToolsetRegistry/Source/ToolsetRegistry/Private/ToolsetRegistry/ToolCallAsyncResultVoid.cpp
+++ b/Engine/Plugins/Experimental/ToolsetRegistry/Source/ToolsetRegistry/Private/ToolsetRegistry/ToolCallAsyncResultVoid.cpp
@@ -2,10 +2,18 @@
 
 #include "ToolsetRegistry/ToolCallAsyncResultVoid.h"
 
+#include "Dom/JsonObject.h"
 #include "Dom/JsonValue.h"
 #include "Templates/SharedPointer.h"
 
+#include "ToolsetRegistry/JsonSchema.h"
+
 TSharedPtr<FJsonValue> UToolCallAsyncResultVoid::GetValueAsJson() const
 {
 	return bIsComplete ? MakeShared<FJsonValueNull>() : TSharedPtr<FJsonValue>();
 }
+
+TSharedRef<FJsonObject> UToolCallAsyncResultVoid::GetValueJsonSchema()
+{
+	return UE::ToolsetRegistry::Internal::CreateJsonSchema(TEXT("Always null"), EJson::Null);
+}
\ No newline at end of file
diff --git a/Engine/Plugins/Experimental/ToolsetRegistry/Source/ToolsetRegistry/Public/ToolsetRegistry/ToolCallAsyncResult.h b/Engine/Plugins/Experimental/ToolsetRegistry/Source/ToolsetRegistry/Public/ToolsetRegistry/ToolCallAsyncResult.h
--- a/Engine/Plugins/Experimental/ToolsetRegistry/Source/ToolsetRegistry/Public/ToolsetRegistry/ToolCallAsyncResult.h
+++ b/Engine/Plugins/Experimental/ToolsetRegistry/Source/ToolsetRegistry/Public/ToolsetRegistry/ToolCallAsyncResult.h
@@ -5,12 +5,16 @@
 #include "Async/Mutex.h"
 #include "Containers/UnrealString.h"
 #include "Delegates/DelegateCombinations.h"
+#include "Dom/JsonObject.h"
 #include "Dom/JsonValue.h"
+#include "Misc/NotNull.h"
 #include "Templates/Function.h"
 #include "Templates/SharedPointer.h"
+#include "Templates/SubclassOf.h"
 #include "Templates/UnrealTemplate.h"
 #include "UObject/Object.h"
 #include "UObject/ObjectMacros.h"
+#include "UObject/UnrealType.h"
 
 #include "ToolCallAsyncResult.generated.h"
 
@@ -138,9 +142,21 @@ private:
 	// Guards: bIsComplete.
 	mutable UE::FMutex CompletedLock;
 
+public:
+	// Get the JSON schema for an async result class' value.
+	static TSharedRef<FJsonObject> GetValueJsonSchema(TSubclassOf<UToolCallAsyncResult> Class);
+
+	// Check whether a property references an async result class.
+	// Returns the async result UClass if the property is an FObjectProperty whose class
+	// derives from UToolCallAsyncResult, nullptr otherwise.
+	static UClass* MatchesProperty(TNotNull<const FProperty*> Property);
+
 private:
 	// Run a function on the main thread as soon as possible.
 	static void RunOnMainThread(TUniqueFunction<void()>&& Function);
+
+private:
+	static FString ValuePropertyName;
 };
 
 #undef UE_API
\ No newline at end of file
diff --git a/Engine/Plugins/Experimental/ToolsetRegistry/Source/ToolsetRegistry/Public/ToolsetRegistry/ToolCallAsyncResultString.h b/Engine/Plugins/Experimental/ToolsetRegistry/Source/ToolsetRegistry/Public/ToolsetRegistry/ToolCallAsyncResultString.h
--- a/Engine/Plugins/Experimental/ToolsetRegistry/Source/ToolsetRegistry/Public/ToolsetRegistry/ToolCallAsyncResultString.h
+++ b/Engine/Plugins/Experimental/ToolsetRegistry/Source/ToolsetRegistry/Public/ToolsetRegistry/ToolCallAsyncResultString.h
@@ -20,7 +20,7 @@ public:
 		return MaybeBroadcastSuccessfulCompletion(FString(InValue), Value);
 	}
 
-	// Value of the result if it has been completed. This should only be set via SetValue().
+	// Value of the result.
 	UPROPERTY(BlueprintReadOnly, Category = "ToolsetRegistry")
 	FString Value;
 };
diff --git a/Engine/Plugins/Experimental/ToolsetRegistry/Source/ToolsetRegistry/Public/ToolsetRegistry/ToolCallAsyncResultVoid.h b/Engine/Plugins/Experimental/ToolsetRegistry/Source/ToolsetRegistry/Public/ToolsetRegistry/ToolCallAsyncResultVoid.h
--- a/Engine/Plugins/Experimental/ToolsetRegistry/Source/ToolsetRegistry/Public/ToolsetRegistry/ToolCallAsyncResultVoid.h
+++ b/Engine/Plugins/Experimental/ToolsetRegistry/Source/ToolsetRegistry/Public/ToolsetRegistry/ToolCallAsyncResultVoid.h
@@ -2,6 +2,7 @@
 
 #pragma once
 
+#include "Dom/JsonObject.h"
 #include "Dom/JsonValue.h"
 #include "Templates/SharedPointer.h"
 #include "UObject/ObjectMacros.h"
@@ -27,6 +28,10 @@ public:
 	}
 
 	UE_API virtual TSharedPtr<FJsonValue> GetValueAsJson() const override;
+
+public:
+	// Get the JSON schema for this class's value.
+	static TSharedRef<FJsonObject> GetValueJsonSchema();
 };
 
 #undef UE_API 
