[AI Assistant] Module with utilities to access properties and their values.

It looks like we have similar operations with properties that are required in few
places across the toolset registry plugin, so this starts a new module to provide
common property access functionality.

#Synced-CL 51118981

diff --git a/Engine/Plugins/Experimental/ToolsetRegistry/Source/ToolsetRegistry/Private/Tests/PropertyAccessorsTest.cpp b/Engine/Plugins/Experimental/ToolsetRegistry/Source/ToolsetRegistry/Private/Tests/PropertyAccessorsTest.cpp
new file mode 100644
--- /dev/null
+++ b/Engine/Plugins/Experimental/ToolsetRegistry/Source/ToolsetRegistry/Private/Tests/PropertyAccessorsTest.cpp
@@ -0,0 +1,81 @@
+// Copyright Epic Games, Inc. All Rights Reserved.
+
+#include "Tests/PropertyAccessorsTest.h"
+
+#include "Containers/UnrealString.h"
+#include "Misc/AutomationTest.h"
+#include "UObject/Class.h"
+
+#include "ToolsetRegistry/PropertyAccessors.h"
+
+#if WITH_DEV_AUTOMATION_TESTS
+
+BEGIN_DEFINE_SPEC(FToolsetRegistryPropertyAccessorsTest,
+	"AI.ToolsetRegistry.PropertyAccessors",
+	EAutomationTestFlags::EditorContext | EAutomationTestFlags::EngineFilter)
+END_DEFINE_SPEC(FToolsetRegistryPropertyAccessorsTest)
+
+void FToolsetRegistryPropertyAccessorsTest::Define()
+{
+	using namespace UE::ToolsetRegistry::Internal;
+
+	Describe(TEXT("PropertyValueAsObject"), [this]()
+	{
+		It(TEXT("Should return a matching object property for null value"), [this]()
+		{
+			FProperty* Property =
+				FToolsetRegistryPropertyValueAsObjectTestStruct::FindTestObjectProperty(*this);
+			if (!Property) return;
+
+			FToolsetRegistryPropertyValueAsObjectTestStruct TestStruct;
+			auto MaybeObject =
+				PropertyValueAsObject<UToolsetRegistryPropertyAccessorsTestObject>(
+					Property, &TestStruct.TestObject);
+			if (!TestTrue(TEXT("Matches object"), MaybeObject.IsSet())) return;
+			TestTrue(TEXT("Matches value"), (*MaybeObject).Get() == TestStruct.TestObject.Get());
+		});
+
+		It(TEXT("Should return a matching object property for non-null value"), [this]()
+		{
+			FProperty* Property =
+				FToolsetRegistryPropertyValueAsObjectTestStruct::FindTestObjectProperty(*this);
+			if (!Property) return;
+
+			FToolsetRegistryPropertyValueAsObjectTestStruct TestStruct;
+			TestStruct.TestObject = NewObject<UToolsetRegistryPropertyAccessorsTestObject>();
+			auto MaybeObject =
+				PropertyValueAsObject<UToolsetRegistryPropertyAccessorsTestObject>(
+					Property, &TestStruct.TestObject);
+			if (!TestTrue(TEXT("Matches object"), MaybeObject.IsSet())) return;
+			TestTrue(TEXT("Matches value"), (*MaybeObject).Get() == TestStruct.TestObject.Get());
+		});
+
+		It(TEXT("Should not return a non-matching object property value"), [this]()
+		{
+			FProperty* Property =
+				FToolsetRegistryPropertyValueAsObjectTestStruct::FindTestAssetManager(*this);
+			if (!Property) return;
+
+			FToolsetRegistryPropertyValueAsObjectTestStruct TestStruct;
+			auto MaybeObject =
+				PropertyValueAsObject<UToolsetRegistryPropertyAccessorsTestObject>(
+					Property, &TestStruct.TestAssetManager);
+			TestFalse(TEXT("Does not match object"), MaybeObject.IsSet());
+		});
+
+		It(TEXT("Should not return a non-matching property value"), [this]()
+		{
+			FProperty* Property =
+				FToolsetRegistryPropertyValueAsObjectTestStruct::FindIntProperty(*this);
+			if (!Property) return;
+
+			FToolsetRegistryPropertyValueAsObjectTestStruct TestStruct;
+			auto MaybeObject =
+				PropertyValueAsObject<UToolsetRegistryPropertyAccessorsTestObject>(
+					Property, &TestStruct.TestObject);
+			TestFalse(TEXT("Does not match object"), MaybeObject.IsSet());
+		});
+	});
+}
+
+#endif  // WITH_DEV_AUTOMATION_TESTS
\ No newline at end of file
diff --git a/Engine/Plugins/Experimental/ToolsetRegistry/Source/ToolsetRegistry/Private/Tests/PropertyAccessorsTest.h b/Engine/Plugins/Experimental/ToolsetRegistry/Source/ToolsetRegistry/Private/Tests/PropertyAccessorsTest.h
new file mode 100644
--- /dev/null
+++ b/Engine/Plugins/Experimental/ToolsetRegistry/Source/ToolsetRegistry/Private/Tests/PropertyAccessorsTest.h
@@ -0,0 +1,63 @@
+// Copyright Epic Games, Inc. All Rights Reserved.
+
+#pragma once
+
+#include "Containers/UnrealString.h"
+#include "Engine/AssetManager.h"
+#include "Misc/AutomationTest.h"
+#include "UObject/Object.h"
+#include "UObject/ObjectMacros.h"
+#include "UObject/ObjectPtr.h"
+
+#include "PropertyAccessorsTest.generated.h"
+
+UCLASS(BlueprintType, MinimalAPI, Hidden)
+class UToolsetRegistryPropertyAccessorsTestObject : public UObject
+{
+	GENERATED_BODY()
+};
+
+USTRUCT(BlueprintType)
+struct FToolsetRegistryPropertyValueAsObjectTestStruct
+{
+	GENERATED_BODY()
+
+public:
+	UPROPERTY()
+	TObjectPtr<UToolsetRegistryPropertyAccessorsTestObject> TestObject;
+
+	UPROPERTY()
+	TObjectPtr<UAssetManager> TestAssetManager;
+
+	UPROPERTY()
+	int TestInt;
+
+public:
+	static FProperty* FindPropertyByName(
+		FAutomationTestBase& Test, const FString& PropertyName)
+	{
+		FProperty* Property = StaticStruct()->FindPropertyByName(*PropertyName);
+		if (!Test.TestTrue(
+				FString::Printf(TEXT("Has %s"), *PropertyName),
+				Property != nullptr))
+		{
+			return nullptr;
+		}
+		return Property;
+	}
+
+	static FProperty* FindTestObjectProperty(FAutomationTestBase& Test)
+	{
+		return FindPropertyByName(Test, TEXT("TestObject"));
+	}
+
+	static FProperty* FindTestAssetManager(FAutomationTestBase& Test)
+	{
+		return FindPropertyByName(Test, TEXT("TestAssetManager"));
+	}
+
+	static FProperty* FindIntProperty(FAutomationTestBase& Test)
+	{
+		return FindPropertyByName(Test, TEXT("TestInt"));
+	}
+};
\ No newline at end of file
diff --git a/Engine/Plugins/Experimental/ToolsetRegistry/Source/ToolsetRegistry/Private/ToolsetRegistry/PropertyAccessors.cpp b/Engine/Plugins/Experimental/ToolsetRegistry/Source/ToolsetRegistry/Private/ToolsetRegistry/PropertyAccessors.cpp
new file mode 100644
--- /dev/null
+++ b/Engine/Plugins/Experimental/ToolsetRegistry/Source/ToolsetRegistry/Private/ToolsetRegistry/PropertyAccessors.cpp
@@ -0,0 +1,24 @@
+// Copyright Epic Games, Inc. All Rights Reserved.
+
+#include "ToolsetRegistry/PropertyAccessors.h"
+
+#include "Misc/Optional.h"
+#include "UObject/Object.h"
+#include "UObject/ObjectPtr.h"
+#include "UObject/UnrealType.h"
+
+namespace UE::ToolsetRegistry::Internal
+{
+	TOptional<TObjectPtr<UObject>> PropertyValueAsObject(
+		TNotNull<UClass*> Class, TNotNull<FProperty*> Property, TNotNull<const void*> ValueAddress)
+	{
+		TOptional<TObjectPtr<UObject>> Object;
+		if (FObjectPropertyBase* ObjectProperty = CastField<FObjectPropertyBase>(Property);
+			ObjectProperty && ObjectProperty->PropertyClass &&
+			ObjectProperty->PropertyClass->IsChildOf(Class))
+		{
+			Object.Emplace(ObjectProperty->GetObjectPropertyValue(ValueAddress));
+		}
+		return Object;
+	}
+}
\ No newline at end of file
diff --git a/Engine/Plugins/Experimental/ToolsetRegistry/Source/ToolsetRegistry/Private/ToolsetRegistry/PropertyAccessors.h b/Engine/Plugins/Experimental/ToolsetRegistry/Source/ToolsetRegistry/Private/ToolsetRegistry/PropertyAccessors.h
new file mode 100644
--- /dev/null
+++ b/Engine/Plugins/Experimental/ToolsetRegistry/Source/ToolsetRegistry/Private/ToolsetRegistry/PropertyAccessors.h
@@ -0,0 +1,31 @@
+// Copyright Epic Games, Inc. All Rights Reserved.
+
+#pragma once
+
+#include "Misc/NotNull.h"
+#include "Misc/Optional.h"
+#include "UObject/Object.h"
+#include "UObject/ObjectPtr.h"
+#include "UObject/UnrealType.h"
+
+namespace UE::ToolsetRegistry::Internal
+{
+	// Get a property value as an object of a specific type.
+	// The returned optional is populated if the property matches the specified UClass.
+	TOptional<TObjectPtr<UObject>> PropertyValueAsObject(
+		TNotNull<UClass*> Class, TNotNull<FProperty*> Property,
+		TNotNull<const void*> ValueAddress);
+
+	// Get a property value as an object of a specific type.
+	// The returned optional is populated if the property matches the UClass ObjectT.
+	template<typename ObjectT>
+	TOptional<TObjectPtr<ObjectT>> PropertyValueAsObject(
+		TNotNull<FProperty*> Property, TNotNull<const void*> ValueAddress)
+	{
+		TOptional<TObjectPtr<ObjectT>> Object;
+		auto MaybeObjectBase = PropertyValueAsObject(
+			TNotNull<UClass*>(ObjectT::StaticClass()), Property, ValueAddress);
+		if (MaybeObjectBase.IsSet()) Object.Emplace(Cast<ObjectT>(*MaybeObjectBase));
+		return Object;
+	}
+}
\ No newline at end of file
