[AI Assistant] Ensure async result handler broadcasts on completed result.

Also, updated tests to catch this,

#Synced-CL 51169544

diff --git a/Engine/Plugins/Experimental/ToolsetRegistry/Source/ToolsetRegistry/Private/Tests/ToolCallAsyncResultTest.cpp b/Engine/Plugins/Experimental/ToolsetRegistry/Source/ToolsetRegistry/Private/Tests/ToolCallAsyncResultTest.cpp
--- a/Engine/Plugins/Experimental/ToolsetRegistry/Source/ToolsetRegistry/Private/Tests/ToolCallAsyncResultTest.cpp
+++ b/Engine/Plugins/Experimental/ToolsetRegistry/Source/ToolsetRegistry/Private/Tests/ToolCallAsyncResultTest.cpp
@@ -6,6 +6,7 @@
 #include "Dom/JsonObject.h"
 #include "Misc/AssertionMacros.h"
 #include "Misc/AutomationTest.h"
+#include "Misc/Optional.h"
 #include "HAL/Event.h"
 #include "Templates/FunctionFwd.h"
 #include "Templates/SharedPointer.h"
@@ -25,15 +26,24 @@
 // On destruction or when calling Validate(), validates the completion of a UToolCallAsyncResult.
 class FToolCallAsyncResultCompletionValidator
 {
+public:
+	enum class EResultHandlerMode
+	{
+		// Enable the result handler on construction.
+		Enabled,
+		// Delay construction of the result handler until it's accessed.
+		Deferred,
+	};
+
 public:
 	FToolCallAsyncResultCompletionValidator(
 		FAutomationTestBase& TestCaseIn,
-		TObjectPtr<UToolCallAsyncResult> ResultIn) :
+		TObjectPtr<UToolCallAsyncResult> ResultIn,
+		EResultHandlerMode ResultHandlerMode = EResultHandlerMode::Enabled) :
 		TestCase(TestCaseIn),
-		Result(ResultIn),
-		ResultHandler(UToolCallAsyncResultFutureHandler::Create(ResultIn)),
-		ResultHandlerFuture(ResultHandler->GetValueAsJson())
+		Result(ResultIn)
 	{
+		if (ResultHandlerMode == EResultHandlerMode::Enabled) EnsureResultHandler();
 	}
 
 	virtual ~FToolCallAsyncResultCompletionValidator()
@@ -93,15 +103,26 @@ public:
 		return *this;
 	}
 
-	TStrongObjectPtr<UToolCallAsyncResultFutureHandler> GetResultHandler() const
+	// Ensure a result handler is created and subscribe to the result.
+	void EnsureResultHandler()
+	{
+		if (!ResultHandler)
+		{
+			ResultHandler = UToolCallAsyncResultFutureHandler::Create(Result.Get());
+			ResultHandlerFuture.Emplace(ResultHandler->GetValueAsJson());
+		}
+	}
+
+	TStrongObjectPtr<UToolCallAsyncResultFutureHandler> GetResultHandler()
 	{
+		EnsureResultHandler();
 		return ResultHandler;
 	}
 
-	TFuture<UE::ToolsetRegistry::FJsonValueOrError>&
-	GetResultHandlerFuture()
+	TFuture<UE::ToolsetRegistry::FJsonValueOrError>& GetResultHandlerFuture()
 	{
-		return ResultHandlerFuture;
+		EnsureResultHandler();
+		return ResultHandlerFuture.GetValue();
 	}
 
 	FAutomationTestBase& GetTestCase() { return TestCase; }
@@ -127,15 +148,16 @@ public:
 		}
 
 		// Verify the OnCompleted event fired and was handled by the handler.
-		if (TestCase.TestTrue(
-				TEXT("Handler was signaled"), ResultHandlerFuture.IsReady()) &&
+		if (ResultHandlerFuture.IsSet() && 
 			TestCase.TestTrue(
-				TEXT("Handler result has value"), ResultHandlerFuture.Get().HasValue()))
+				TEXT("Handler was signaled"), ResultHandlerFuture->IsReady()) &&
+			TestCase.TestTrue(
+				TEXT("Handler result has value"), ResultHandlerFuture->Get().HasValue()))
 		{
 			TestCase.TestEqual(
 				TEXT("Handler result value"),
 				UE::ToolsetRegistry::Internal::JsonToString(
-					ResultHandlerFuture.Get().GetValue().ToSharedRef()),
+					ResultHandlerFuture->Get().GetValue().ToSharedRef()),
 				ExpectedJsonValueString);
 		}
 	}
@@ -154,14 +176,15 @@ public:
 			TEXT("Result has no JSON string"), Result->GetValueAsJsonString(), TEXT(""));
 
 		// Verify the OnCompleted event fired and was handled by the handler.
-		if (TestCase.TestTrue(TEXT("Handler was signaled"), ResultHandlerFuture.IsReady()) &&
+		if (ResultHandlerFuture.IsSet() && 
+			TestCase.TestTrue(TEXT("Handler was signaled"), ResultHandlerFuture->IsReady()) &&
 			TestCase.TestTrue(
 				TEXT("Handler result has error"),
-				ResultHandlerFuture.Get().HasError()))
+				ResultHandlerFuture->Get().HasError()))
 		{
 			TestCase.TestEqual(
 				TEXT("Handler result error"),
-				ResultHandlerFuture.Get().GetError(), ExpectedHandlerError);
+				ResultHandlerFuture->Get().GetError(), ExpectedHandlerError);
 		}
 	}
 
@@ -180,14 +203,14 @@ public:
 		// Verify the OnCompleted event has not fired.
 		TestCase.TestFalse(
 			TEXT("Result handler future is not complete"),
-			ResultHandlerFuture.IsReady());
+			ResultHandlerFuture.IsSet() && ResultHandlerFuture->IsReady());
 	}
 
 private:
 	FAutomationTestBase& TestCase;
 	TStrongObjectPtr<UToolCallAsyncResult> Result;
 	TStrongObjectPtr<UToolCallAsyncResultFutureHandler> ResultHandler;
-	TFuture<UE::ToolsetRegistry::FJsonValueOrError> ResultHandlerFuture;
+	TOptional<TFuture<UE::ToolsetRegistry::FJsonValueOrError>> ResultHandlerFuture;
 	bool bShouldComplete = true;
 	FString ExpectedError;
 	TSharedPtr<FJsonValue> ExpectedJsonValue;
@@ -333,7 +356,7 @@ void FToolCallAsyncResultBaseSpec::Define()
 		Result->SetError(FakeError);
 		FToolCallAsyncResultCompletionValidator Validator(*this, Result);
 		Validator.ExpectError(FakeError);
-		TestTrue(TEXT("Is complete"), Result->BroadcastOnCompletedIfComplete());
+		TestTrue(TEXT("Is complete"), Result->bIsComplete);
 	});
 
 	It(TEXT("Should fail when an empty error is set"), [this]()
@@ -359,8 +382,8 @@ void FToolCallAsyncResultBaseSpec::Define()
 
 	It(TEXT("Should complete when an error is set from another thread"), [this]()
 	{
-		auto Validator = MakeShared<
-			FToolCallAsyncResultCompletionValidator>(*this, Result);
+		auto Validator =
+			MakeShared<FToolCallAsyncResultCompletionValidator>(*this, Result);
 		Validator->ExpectError(FakeError);
 
 		AddCommand(
@@ -373,7 +396,10 @@ void FToolCallAsyncResultBaseSpec::Define()
 	It(TEXT("Should send completion with an error on broadcast from another thread"), [this]()
 	{
 		Result->SetError(FakeError);
-		auto Validator = MakeShared<FToolCallAsyncResultCompletionValidator>(*this, Result);
+		auto Validator =
+			MakeShared<FToolCallAsyncResultCompletionValidator>(
+				*this, Result,
+				FToolCallAsyncResultCompletionValidator::EResultHandlerMode::Deferred);
 		Validator->ExpectError(FakeError);
 
 		AddCommand(
@@ -414,7 +440,6 @@ void FToolCallAsyncResultBaseSpec::Define()
 		SetResultValue();
 		FToolCallAsyncResultCompletionValidator Validator(*this, Result);
 		Validator.ExpectJsonValue(ExpectedJsonValue);
-		Result->BroadcastOnCompletedIfComplete();
 	});
 
 	It(TEXT("Should fail when a value is set more than once"), [this]()
@@ -444,8 +469,10 @@ void FToolCallAsyncResultBaseSpec::Define()
 	It(TEXT("Should send completion with a value on broadcast from another thread"), [this]()
 	{
 		SetResultValue();
-		auto Validator = MakeShared<
-			FToolCallAsyncResultCompletionValidator>(*this, Result);
+		auto Validator =
+			MakeShared<FToolCallAsyncResultCompletionValidator>(
+				*this, Result,
+				FToolCallAsyncResultCompletionValidator::EResultHandlerMode::Deferred);
 		Validator->ExpectJsonValue(ExpectedJsonValue);
 
 		AddCommand(
diff --git a/Engine/Plugins/Experimental/ToolsetRegistry/Source/ToolsetRegistry/Private/ToolsetRegistry/ToolCallAsyncResultFutureHandler.cpp b/Engine/Plugins/Experimental/ToolsetRegistry/Source/ToolsetRegistry/Private/ToolsetRegistry/ToolCallAsyncResultFutureHandler.cpp
--- a/Engine/Plugins/Experimental/ToolsetRegistry/Source/ToolsetRegistry/Private/ToolsetRegistry/ToolCallAsyncResultFutureHandler.cpp
+++ b/Engine/Plugins/Experimental/ToolsetRegistry/Source/ToolsetRegistry/Private/ToolsetRegistry/ToolCallAsyncResultFutureHandler.cpp
@@ -37,6 +37,7 @@ void UToolCallAsyncResultFutureHandler::Subscribe(UToolCallAsyncResult* ResultIn
 	check(!Result);
 	Result = ResultIn;
 	Result->OnCompleted.AddUniqueDynamic(this, &ThisClass::OnCompleted);
+	Result->BroadcastOnCompletedIfComplete();
 }
 
 void UToolCallAsyncResultFutureHandler::Unsubscribe()
