[AI Assistant] Conversion of UToolCallAsyncResult::Value to JSON schema / JSON.

Functions that return UToolCallAsyncResult should be exposed as tools returning
UToolCallAsyncResult::Value . This change hooks the UFunction to JSON schema
and FProperty return value to JSON value conversion to use the object value's type
FProperty to generate the schema and the member function
GetValueAsJsonString() to get the value of the object if one is present.

#Synced-CL 51169544

diff --git a/Engine/Plugins/Experimental/ToolsetRegistry/Source/ToolsetRegistry/Private/Tests/ToolsetJsonTest.cpp b/Engine/Plugins/Experimental/ToolsetRegistry/Source/ToolsetRegistry/Private/Tests/ToolsetJsonTest.cpp
--- a/Engine/Plugins/Experimental/ToolsetRegistry/Source/ToolsetRegistry/Private/Tests/ToolsetJsonTest.cpp
+++ b/Engine/Plugins/Experimental/ToolsetRegistry/Source/ToolsetRegistry/Private/Tests/ToolsetJsonTest.cpp
@@ -3,10 +3,17 @@
 #include "ToolsetJsonTest.h"
 
 #include "CoreMinimal.h"
+#include "Containers/Array.h"
+#include "Containers/ArrayView.h"
 #include "JsonObjectConverter.h"
 #include "Misc/AutomationTest.h"
 
+#include "ToolsetRegistry/JsonConversion.h"
+#include "ToolsetRegistry/JsonSchema.h"
 #include "ToolsetRegistry/ToolsetJson.h"
+#include "ToolsetRegistry/ToolCallAsyncResult.h"
+#include "ToolsetRegistry/ToolCallAsyncResultString.h"
+#include "Tests/ToolCallAsyncResultTest.h"
 
 #if WITH_DEV_AUTOMATION_TESTS
 
@@ -171,14 +178,43 @@ namespace
 			TestColorData(Test, *Property, Expected);
 		}
 	}
+
+	TSharedPtr<FJsonObject> TestHasObjectField(
+		FAutomationTestBase& Test, const TSharedPtr<FJsonObject>& Schema,
+		const FString& RootObjectName, const TArrayView<FString>& FieldPath)
+	{
+		const TSharedPtr<FJsonObject>* MaybeObjectField = nullptr;
+		TSharedPtr<FJsonObject> OutputObject;
+		verify(FieldPath.Num() > 0);
+		FString FieldName = FieldPath[0];
+		FString ParentName =
+			FieldPath.Num() > 1
+			? FString::Printf(TEXT("%s.%s"), *RootObjectName, *FieldName)
+			: RootObjectName;
+		if (Schema.IsValid() &&
+			Test.TestTrue(
+				FString::Printf(TEXT("%s has field %s"), *ParentName, *FieldName),
+				Schema->TryGetObjectField(*FieldName, MaybeObjectField) &&
+				MaybeObjectField && MaybeObjectField->IsValid()) &&
+			MaybeObjectField)
+		{
+			OutputObject =
+				FieldPath.Num() > 1
+				? TestHasObjectField(
+					Test, *MaybeObjectField, ParentName, FieldPath.Right(FieldPath.Num() - 1))
+				: *MaybeObjectField;
+		}
+		return OutputObject;
+	}
 }
 
 BEGIN_DEFINE_SPEC(FToolsetJsonSpec, "AI.ToolsetRegistry.ToolsetJsonSpec",
-	EAutomationTestFlags::EditorContext | EAutomationTestFlags::EngineFilter)	
+	EAutomationTestFlags::EditorContext | EAutomationTestFlags::EngineFilter)
 END_DEFINE_SPEC(FToolsetJsonSpec)
 
 void FToolsetJsonSpec::Define()
 {
+	using namespace UE::ToolsetRegistry::Internal;
 	using namespace UE::ToolsetRegistry::Internal::ToolsetJson;
 
 	It("Converts UObject* property to reference schema", [this]()
@@ -402,6 +438,10 @@ void FToolsetJsonSpec::Define()
 					{
 						TestColorSchema(*this, *TestProperty, JsonPropertyName == FString(TEXT("testColor")));
 					}
+					else if (JsonPropertyName == TEXT("testToolCallAsyncResult"))
+					{
+						continue;  // Ignored as references to the base class are always null.
+					}
 					else
 					{
 						TestObjectRefSchema(*this, *TestProperty, nullptr, false);
@@ -444,6 +484,10 @@ void FToolsetJsonSpec::Define()
 				TestColorDataInStruct(
 					*this, JsonObject, JsonPropertyName, TestIn.TestLinearColor.ToFColor(false));
 			}
+			else if (JsonPropertyName == TEXT("testToolCallAsyncResult"))
+			{
+				continue;  // Does not support conversion.
+			}
 			else
 			{
 				TestObjectRefDataInStruct(
@@ -549,6 +593,46 @@ void FToolsetJsonSpec::Define()
 		}
 	});
 
+	It(TEXT("Converts ToolCallAsyncResult return value to JSON schema"), [this]()
+	{
+		UFunction* Function = UToolCallAsyncResultReturnPropertyTest::FindTestFunction();
+		if (!TestTrue(TEXT("Has test function"), Function != nullptr)) return;
+
+		const TSharedPtr<FJsonObject> Schema = StructToJsonSchema(Function);
+		const TSharedPtr<FJsonObject> ReturnValueSchema =
+			TestHasObjectField(
+				*this, Schema, TEXT("Function schema"),
+				TArray<FString>{ TEXT("outputSchema"), TEXT("properties"), TEXT("returnValue") });
+		if (!ReturnValueSchema) return;
+
+		TestEqual(
+			TEXT("Matches expected schema"),
+			JsonToString(ReturnValueSchema.ToSharedRef()),
+			JsonToString(PropertyToJsonSchema(Function->GetReturnProperty()).ToSharedRef()));
+			
+	});
+
+	It(TEXT("Converts ToolCallAsyncResult to JSON"), [this]()
+	{
+		TObjectPtr<UToolCallAsyncResultString> ToolCallAsyncResult =
+			NewObject<UToolCallAsyncResultString>();
+		ToolCallAsyncResult->SetValue(TEXT("hello world"));
+
+		FToolsetJsonTest TestIn;
+		TestIn.TestToolCallAsyncResult = ToolCallAsyncResult;
+		UStruct* Struct = FToolsetJsonTest::StaticStruct();
+		FProperty* Property = Struct->FindPropertyByName(TEXT("TestToolCallAsyncResult"));
+		if (!TestTrue(TEXT("Has property"), Property != nullptr)) return;
+
+		TSharedPtr<FJsonValue> JsonValue = PropertyToJsonData(
+			Property, &TestIn.TestToolCallAsyncResult);
+		if (!TestTrue(TEXT("Converted return value to JSON"), JsonValue.IsValid())) return;
+
+		TestEqual(
+			TEXT("Matches value JSON"),
+			JsonToString(JsonValue.ToSharedRef()),
+			ToolCallAsyncResult->GetValueAsJsonString());
+	});
 }
 
 #endif
diff --git a/Engine/Plugins/Experimental/ToolsetRegistry/Source/ToolsetRegistry/Private/Tests/ToolsetJsonTest.h b/Engine/Plugins/Experimental/ToolsetRegistry/Source/ToolsetRegistry/Private/Tests/ToolsetJsonTest.h
--- a/Engine/Plugins/Experimental/ToolsetRegistry/Source/ToolsetRegistry/Private/Tests/ToolsetJsonTest.h
+++ b/Engine/Plugins/Experimental/ToolsetRegistry/Source/ToolsetRegistry/Private/Tests/ToolsetJsonTest.h
@@ -6,6 +6,8 @@
 #include "Templates/SharedPointer.h"
 #include "Templates/SubclassOf.h"
 
+#include "ToolsetRegistry/ToolCallAsyncResult.h"
+
 #include "ToolsetJsonTest.generated.h"
 
 
@@ -54,4 +56,7 @@ public:
 
 	UPROPERTY()
 	FLinearColor TestLinearColor = FLinearColor::White;
+
+	UPROPERTY()
+	TObjectPtr<UToolCallAsyncResult> TestToolCallAsyncResult;
 };
diff --git a/Engine/Plugins/Experimental/ToolsetRegistry/Source/ToolsetRegistry/Private/ToolsetRegistry/ToolsetJson.cpp b/Engine/Plugins/Experimental/ToolsetRegistry/Source/ToolsetRegistry/Private/ToolsetRegistry/ToolsetJson.cpp
--- a/Engine/Plugins/Experimental/ToolsetRegistry/Source/ToolsetRegistry/Private/ToolsetRegistry/ToolsetJson.cpp
+++ b/Engine/Plugins/Experimental/ToolsetRegistry/Source/ToolsetRegistry/Private/ToolsetRegistry/ToolsetJson.cpp
@@ -9,6 +9,9 @@
 #include "Modules/ModuleManager.h"
 #include "Misc/AssertionMacros.h"
 
+#include "ToolsetRegistry/PropertyAccessors.h"
+#include "ToolsetRegistry/ToolCallAsyncResult.h"
+
 namespace
 {
 	// Use RequiredModules::EnsureLoaded() to lazily load required modules.
@@ -94,7 +97,12 @@ bool CustomJsonSchema(
 	const FProperty* Property, const FString& ParameterDefaultString, const TSharedRef<FJsonObject>& OutputSchema)
 {
 	RequiredModules::EnsureLoaded();
-	if (IsReferenceLikeProperty(Property))
+	if (UClass* Class = UToolCallAsyncResult::MatchesProperty(Property))
+	{
+		FJsonObject::Duplicate(UToolCallAsyncResult::GetValueJsonSchema(Class), OutputSchema);
+		return true;
+	}
+	else if (IsReferenceLikeProperty(Property))
 	{
 		if (ParameterDefaultString.IsEmpty())
 		{
@@ -230,14 +238,20 @@ bool CustomImport(const TSharedPtr<FJsonValue>& JsonValue, FProperty* Property,
 TSharedPtr<FJsonValue> CustomExport(FProperty* Property, const void* Value)
 {
 	RequiredModules::EnsureLoaded();
-	if (IsReferenceLikeProperty(Property))
+	if (auto Result = PropertyValueAsObject<UToolCallAsyncResult>(Property, Value);
+		Result.IsSet())
+	{
+		return *Result ? (*Result)->GetValueAsJson() : nullptr;
+	}
+	else if (IsReferenceLikeProperty(Property))
 	{
 		FToolsetReference Reference;
-		if (FObjectPropertyBase* ObjectProperty = CastField<FObjectPropertyBase>(Property))
+		if (auto MaybeObject = PropertyValueAsObject<UObject>(Property, Value);
+			MaybeObject.IsSet())
 		{
-			if (const UObject* Object = ObjectProperty->GetObjectPropertyValue(Value))
+			if (*MaybeObject)
 			{
-				Reference.RefPath = Object->GetPathName();
+				Reference.RefPath = (*MaybeObject)->GetPathName();
 			}
 			else
 			{
