[AI Assistant] Some clean up of ObjectFunctionToolCallTest

Move to spec style test in prep for adding async result tests.

#Synced-CL 51169544

diff --git a/Engine/Plugins/Experimental/ToolsetRegistry/Source/ToolsetRegistry/Private/Tests/ObjectFunctionToolCallTest.cpp b/Engine/Plugins/Experimental/ToolsetRegistry/Source/ToolsetRegistry/Private/Tests/ObjectFunctionToolCallTest.cpp
--- a/Engine/Plugins/Experimental/ToolsetRegistry/Source/ToolsetRegistry/Private/Tests/ObjectFunctionToolCallTest.cpp
+++ b/Engine/Plugins/Experimental/ToolsetRegistry/Source/ToolsetRegistry/Private/Tests/ObjectFunctionToolCallTest.cpp
@@ -30,7 +30,8 @@ namespace UE::ToolsetRegistry
 			{
 				FString JsonString =
 					UE::ToolsetRegistry::Internal::JsonToString(JsonObject.ToSharedRef());
-				UE_LOG(LogToolsetRegistry, Display, TEXT("Json \"%s\" is valid -\n%s"), Name, *JsonString);
+				UE_LOG(LogToolsetRegistry, Display,
+					TEXT("Json \"%s\" is valid -\n%s"), Name, *JsonString);
 			}
 			else
 			{
@@ -39,10 +40,12 @@ namespace UE::ToolsetRegistry
 		}
 		
 		// Ensure a tool call result has no error.
-		bool TestHasNoError(FAutomationTestBase& Test, FObjectFunctionToolCall::FToolCallResult& Result)
+		bool TestHasNoError(
+			FAutomationTestBase& Test, FObjectFunctionToolCall::FToolCallResult& Result)
 		{
 			const FString* Error = Result.TryGetError();
-			return !Error || Test.TestEqual(TEXT("Tool call result has no error."), *Error, TEXT(""));
+			return !Error ||
+				Test.TestEqual(TEXT("Tool call result has no error."), *Error, TEXT(""));
 		}
 
 		// Ensure a tool call result has a valid value.
@@ -81,111 +84,104 @@ namespace UE::ToolsetRegistry
 			return JsonObject;
 		}
 	}
+}
 
-	// Test function with required param.
-
-	IMPLEMENT_SIMPLE_AUTOMATION_TEST(
-		FObjectFunctionToolCallTestExecuteToolCallFunctionWithRequiredParam,
-		"AI.ToolsetRegistry.ObjectFunctionToolCall.ExecuteToolCallFunctionWithRequiredParam",
-		ToolsetRegistryTest::Flags);
-
-	bool FObjectFunctionToolCallTestExecuteToolCallFunctionWithRequiredParam::RunTest(
-		const FString& Parameters)
-	{
-		static const FName FunctionName = "TestFuncWithRequiredParam";
-		static constexpr float ClassFloatValue = 20.0f;
-		static constexpr float FunctionInputFloatValue = 3.0f;
-		static const FString FunctionInputParamsJsonString =
-			FString::Printf(TEXT("{\n\"InFloat\": %f\n}"), FunctionInputFloatValue);
-
-		// Create an object and provide it specific internal state. 
-		// (For this test - Using TStrongObjectPtr to prevent GC.)
-
-		const TStrongObjectPtr InstanceObject =
-			TStrongObjectPtr(NewObject<UToolCallFakeClass>());
-
-		InstanceObject->Float = ClassFloatValue;
-
-		// Execute.
-
-		const FObjectFunctionToolCall ToolCall =
-			FObjectFunctionToolCall::ConstructWithNamedClassFunction(
-				InstanceObject.Get(), FunctionName);
-
-		FObjectFunctionToolCall::FFunctionInputParamsJson FunctionInputParamsJson;
-		FunctionInputParamsJson.Set<FString>(FunctionInputParamsJsonString);
-
-		AddCommand(
-			UE::ToolsetRegistry::Internal::FUntilFutureCommand<
-					FObjectFunctionToolCall::FToolCallResult>::Create(
-				ToolCall.Execute(FunctionInputParamsJson),
-				[this, InstanceObject](
-					FObjectFunctionToolCall::FToolCallResult&& ToolCallResult) -> void
-				{
-					TSharedPtr<FJsonObject> JsonObject =
-						TestHasReturnValueField(*this, ToolCallResult, EJson::Number);
-					if (!JsonObject) return;
-					TestNearlyEqual(
-						TEXT("Tool call result value's returnValue"),
-						static_cast<float>(JsonObject->GetNumberField(TEXT("returnValue"))),
-						InstanceObject->TestFuncWithRequiredParam(FunctionInputFloatValue),
-						0.0001f);
-				}));
-
-		return true;
-	}
-
-	// Test function with optional param.
+BEGIN_DEFINE_SPEC(FObjectFunctionToolCallTest, "AI.ToolsetRegistry.ObjectFunctionToolCall",
+	ToolsetRegistryTest::Flags)
+END_DEFINE_SPEC(FObjectFunctionToolCallTest)
 
-	IMPLEMENT_SIMPLE_AUTOMATION_TEST(
-		FObjectFunctionToolCallTestExecuteToolCallFunctionWithOptionalParam,
-		"AI.ToolsetRegistry.ObjectFunctionToolCall.ExecuteToolCallFunctionOptionalParam",
-		ToolsetRegistryTest::Flags);
+void FObjectFunctionToolCallTest::Define()
+{
+	using namespace UE::ToolsetRegistry;
 
-	bool FObjectFunctionToolCallTestExecuteToolCallFunctionWithOptionalParam::RunTest(
-		const FString& Parameters)
+	Describe(TEXT("ExecuteToolCall"), [this]()
 	{
-		static const FName FunctionName = "TestFuncWithOptionalParam";
-		static constexpr float ClassFloatValue = 20.0f;
-
-		// Create an object and provide it specific internal state. 
-		// (For this test - Using TStrongObjectPtr to prevent GC.)
+		It(TEXT("should execute a function with a required parameter"), [this]()
+		{
+			static const FName FunctionName = TEXT("TestFuncWithRequiredParam");
+			static constexpr float ClassFloatValue = 20.0f;
+			static constexpr float FunctionInputFloatValue = 3.0f;
+			static const FString FunctionInputParamsJsonString =
+				FString::Printf(TEXT("{\n\"InFloat\": %f\n}"), FunctionInputFloatValue);
+
+			// Create an object and provide it specific internal state. 
+			// (For this test - Using TStrongObjectPtr to prevent GC.)
+
+			const TStrongObjectPtr InstanceObject =
+				TStrongObjectPtr(NewObject<UToolCallFakeClass>());
+
+			InstanceObject->Float = ClassFloatValue;
+
+			// Execute.
+
+			const FObjectFunctionToolCall ToolCall =
+				FObjectFunctionToolCall::ConstructWithNamedClassFunction(
+					InstanceObject.Get(), FunctionName);
+
+			FObjectFunctionToolCall::FFunctionInputParamsJson FunctionInputParamsJson;
+			FunctionInputParamsJson.Set<FString>(FunctionInputParamsJsonString);
+
+			AddCommand(
+				UE::ToolsetRegistry::Internal::FUntilFutureCommand<
+						FObjectFunctionToolCall::FToolCallResult>::Create(
+					ToolCall.Execute(FunctionInputParamsJson),
+					[this, InstanceObject](
+						FObjectFunctionToolCall::FToolCallResult&& ToolCallResult) -> void
+					{
+						TSharedPtr<FJsonObject> JsonObject =
+							TestHasReturnValueField(*this, ToolCallResult, EJson::Number);
+						if (!JsonObject) return;
+						TestNearlyEqual(
+							TEXT("Tool call result value's returnValue"),
+							static_cast<float>(JsonObject->GetNumberField(TEXT("returnValue"))),
+							InstanceObject->TestFuncWithRequiredParam(FunctionInputFloatValue),
+							0.0001f);
+					}));
+		});
+
+		It(TEXT("should execute a function with an optional parameter"), [this]()
+		{
+			static const FName FunctionName = TEXT("TestFuncWithOptionalParam");
+			static constexpr float ClassFloatValue = 20.0f;
 
-		const TStrongObjectPtr InstanceObject =
-			TStrongObjectPtr(NewObject<UToolCallFakeClass>());
+			// Create an object and provide it specific internal state. 
+			// (For this test - Using TStrongObjectPtr to prevent GC.)
 
-		InstanceObject->Float = ClassFloatValue;
+			const TStrongObjectPtr InstanceObject =
+				TStrongObjectPtr(NewObject<UToolCallFakeClass>());
 
-		// Execute.
+			InstanceObject->Float = ClassFloatValue;
 
-		const TOptional<FObjectFunctionToolCall> ToolCall =
-			FObjectFunctionToolCall::ConstructWithNamedClassFunction(InstanceObject.Get(), FunctionName);
-		if (!TestTrue(TEXT("Tool call created from class and function name."),
-			ToolCall.IsSet()))
-		{
-			return false;
-		}
+			// Execute.
 
-		AddCommand(
-			UE::ToolsetRegistry::Internal::FUntilFutureCommand<
-					FObjectFunctionToolCall::FToolCallResult>::Create(
-				ToolCall->Execute(),
-				[this, InstanceObject](
-					FObjectFunctionToolCall::FToolCallResult&& ToolCallResult) -> void
-				{
-					TSharedPtr<FJsonObject> JsonObject =
-						TestHasReturnValueField(*this, ToolCallResult, EJson::Number);
-					if (!JsonObject) return;
-
-					TestNearlyEqual(
-						TEXT("Tool call result value's returnValue"),
-						static_cast<float>(JsonObject->GetNumberField(TEXT("returnValue"))),
-						InstanceObject->TestFuncWithOptionalParam(),
-						0.0001f);
-				}));
-		return true;
-	}
+			const TOptional<FObjectFunctionToolCall> ToolCall =
+				FObjectFunctionToolCall::ConstructWithNamedClassFunction(
+					InstanceObject.Get(), FunctionName);
+			if (!TestTrue(TEXT("Tool call created from class and function name."),
+				ToolCall.IsSet()))
+			{
+				return;
+			}
 
+			AddCommand(
+				UE::ToolsetRegistry::Internal::FUntilFutureCommand<
+						FObjectFunctionToolCall::FToolCallResult>::Create(
+					ToolCall->Execute(),
+					[this, InstanceObject](
+						FObjectFunctionToolCall::FToolCallResult&& ToolCallResult) -> void
+					{
+						TSharedPtr<FJsonObject> JsonObject =
+							TestHasReturnValueField(*this, ToolCallResult, EJson::Number);
+						if (!JsonObject) return;
+
+						TestNearlyEqual(
+							TEXT("Tool call result value's returnValue"),
+							static_cast<float>(JsonObject->GetNumberField(TEXT("returnValue"))),
+							InstanceObject->TestFuncWithOptionalParam(),
+							0.0001f);
+					}));
+		});
+	});
 }
 
 #endif
