[AI Assistant] Further simplify creating fullfilled futures with a value or error

Enables the removal of more boilerplate when creating fullfilled futures
with a value or error. Also, added a shared typedef for a JSON value or error
future object.

#Synced-CL 51174756

diff --git a/Engine/Plugins/Experimental/ToolsetRegistry/Source/ToolsetRegistry/Private/Tests/ValueOrErrorFutureTest.cpp b/Engine/Plugins/Experimental/ToolsetRegistry/Source/ToolsetRegistry/Private/Tests/ValueOrErrorFutureTest.cpp
--- a/Engine/Plugins/Experimental/ToolsetRegistry/Source/ToolsetRegistry/Private/Tests/ValueOrErrorFutureTest.cpp
+++ b/Engine/Plugins/Experimental/ToolsetRegistry/Source/ToolsetRegistry/Private/Tests/ValueOrErrorFutureTest.cpp
@@ -17,6 +17,35 @@ void FToolsetRegistryValueOrErrorFutureSpec::Define()
 {
 	using namespace UE::ToolsetRegistry::Internal;
 
+	Describe(TEXT("Make"), [this]()
+	{
+		It(TEXT("should return a value"), [this]()
+		{
+			FString ExpectedValue = TEXT("Yo");
+			TFuture<TValueOrError<FString, FString>> Future =
+				FStringValueOrErrorFuture::Make(MakeValue(ExpectedValue));
+			if (!TestTrue(TEXT("is ready"), Future.IsReady())) return;
+
+			auto ValueOrValue = Future.Get();
+			if (!TestTrue(TEXT("has value"), ValueOrValue.HasValue())) return;
+
+			TestEqual(TEXT("value"), ValueOrValue.GetValue(), ExpectedValue);
+		});
+
+		It(TEXT("should return an error"), [this]()
+		{
+			FString ExpectedError = TEXT("Failed");
+			TFuture<TValueOrError<FString, FString>> Future =
+				FStringValueOrErrorFuture::Make(MakeError(FString(ExpectedError)));
+			if (!TestTrue(TEXT("is ready"), Future.IsReady())) return;
+
+			auto ValueOrError = Future.Get();
+			if (!TestTrue(TEXT("has error"), ValueOrError.HasError())) return;
+
+			TestEqual(TEXT("error"), ValueOrError.GetError(), ExpectedError);
+		});
+	});
+
 	Describe(TEXT("MakeError"), [this]()
 	{
 		It(TEXT("should return an error"), [this]()
diff --git a/Engine/Plugins/Experimental/ToolsetRegistry/Source/ToolsetRegistry/Private/ToolsetRegistry/ValueOrErrorFuture.h b/Engine/Plugins/Experimental/ToolsetRegistry/Source/ToolsetRegistry/Private/ToolsetRegistry/ValueOrErrorFuture.h
--- a/Engine/Plugins/Experimental/ToolsetRegistry/Source/ToolsetRegistry/Private/ToolsetRegistry/ValueOrErrorFuture.h
+++ b/Engine/Plugins/Experimental/ToolsetRegistry/Source/ToolsetRegistry/Private/ToolsetRegistry/ValueOrErrorFuture.h
@@ -7,28 +7,41 @@
 #include "Templates/UnrealTemplate.h"
 #include "Templates/ValueOrError.h"
 
+#include "ToolsetRegistry/JsonValueOrError.h"
+
 namespace UE::ToolsetRegistry::Internal
 {
 	// Reduces the boilerplate required to construct a completed future with a value 
 	// or error 
-	template<typename ValueT>
-	class TValueOrErrorFuture
+	template<typename ValueT, typename ValueOrErrorT>
+	class TValueOrErrorFutureBase
 	{
 	public:
-		using FValueOrError = TValueOrError<ValueT, FString>;
+		using FValueOrError = ValueOrErrorT;
 		using FFuture = TFuture<FValueOrError>;
 
 	public:
+		static FFuture Make(FValueOrError&& ValueOrError)
+		{
+			return MakeFulfilledPromise<FValueOrError>(MoveTemp(ValueOrError)).GetFuture();
+		}
+
 		static FFuture MakeValue(ValueT&& Value)
 		{
-			return MakeFulfilledPromise<FValueOrError>(::MakeValue(MoveTemp(Value))).GetFuture();
+			return Make(::MakeValue(MoveTemp(Value)));
 		}
 
 		static FFuture MakeError(FString&& Error)
 		{
-			return MakeFulfilledPromise<FValueOrError>(::MakeError(MoveTemp(Error))).GetFuture();
+			return Make(::MakeError(MoveTemp(Error)));
 		}
 	};
 
+	template<typename ValueT>
+	using TValueOrErrorFuture = TValueOrErrorFutureBase<ValueT, TValueOrError<ValueT, FString>>;
+
 	using FStringValueOrErrorFuture = TValueOrErrorFuture<FString>;
+
+	using FJsonValueOrErrorFuture =
+		TValueOrErrorFutureBase<TSharedPtr<FJsonValue>, FJsonValueOrError>;
 }
\ No newline at end of file
