[AI Assistant] Support async blueprint tool call UFunctions.

When dispatching a UFunction function that returns a ToolCallAsyncResult
FObjectFunctionToolCall connects ToolCallAsyncResult to the future
returned by FObjectFunctionToolCall::Execute().

Also, cleaned up ObjectFunctionToolCallTests.cpp to reduce the amount of
boilerplate required to add a test case.

#Synced-CL 51174756

diff --git a/Engine/Plugins/Experimental/ToolsetRegistry/Source/ToolsetRegistry/Private/Tests/ObjectFunctionToolCallTest.cpp b/Engine/Plugins/Experimental/ToolsetRegistry/Source/ToolsetRegistry/Private/Tests/ObjectFunctionToolCallTest.cpp
--- a/Engine/Plugins/Experimental/ToolsetRegistry/Source/ToolsetRegistry/Private/Tests/ObjectFunctionToolCallTest.cpp
+++ b/Engine/Plugins/Experimental/ToolsetRegistry/Source/ToolsetRegistry/Private/Tests/ObjectFunctionToolCallTest.cpp
@@ -3,13 +3,16 @@
 #include "ObjectFunctionToolCallTest.h"
 
 #include "Misc/AutomationTest.h"
+#include "Misc/Optional.h"
 #include "UObject/Package.h"
 #include "Serialization/JsonSerializer.h"
 #include "Serialization/JsonWriter.h"
 
 #include "Tests/ToolsetRegistryTestFlags.h"
+#include "Tests/ToolCallAsyncResultTest.h"
 #include "Tests/UntilFutureCommand.h"
 #include "ToolsetRegistryModule.h"
+#include "ToolsetRegistry/JsonValueOrError.h"
 #include "ToolsetRegistry/JsonConversion.h"
 #include "ToolsetRegistry/ObjectFunctionToolCall.h"
 
@@ -38,10 +41,55 @@ namespace UE::ToolsetRegistry
 				UE_LOG(LogToolsetRegistry, Error, TEXT("Json \"%s\" is invalid."), Name);
 			}
 		}
+
+		// Holds a strong reference to an object instance and an object function tool call.
+		template<typename ObjectT>
+		struct FStrongObjectFunctionToolCall
+		{
+			TStrongObjectPtr<ObjectT> InstanceObject;
+			TSharedPtr<FObjectFunctionToolCall> ToolCall;
+		};
+
+		// Create an object function tool call that wraps a function.
+		TSharedPtr<FObjectFunctionToolCall> TestCreateToolCall(
+			FAutomationTestBase& Test, TObjectPtr<UObject> InstanceObject,
+			const FString& FunctionName)
+		{
+			auto ToolCall =
+				FObjectFunctionToolCall::Create(InstanceObject.Get(), *FunctionName);
+			Test.TestTrue(
+				FString::Printf(
+					TEXT("Construct object function tool call for function %s of %s"),
+					*FunctionName, *InstanceObject->GetName()),
+				ToolCall.IsValid());
+			return ToolCall;
+		}
+
+		// Create an object function tool call that wraps a function on an object instance of
+		// type ObjectT.
+		template<typename ObjectT>
+		TOptional<FStrongObjectFunctionToolCall<ObjectT>> TestCreateToolCall(
+			FAutomationTestBase& Test, const FString& FunctionName)
+		{
+			TOptional<FStrongObjectFunctionToolCall<ObjectT>> ToolCallReference;
+			auto InstanceObject = NewObject<ObjectT>();
+			if (Test.TestTrue(
+				FString::Printf(TEXT("Constructed object %s"), *ObjectT::StaticClass()->GetName()),
+				InstanceObject != nullptr))
+			{
+				auto ToolCall = TestCreateToolCall(Test, InstanceObject, FunctionName);
+				if (ToolCall)
+				{
+					ToolCallReference.Emplace();
+					ToolCallReference->ToolCall = ToolCall;
+					ToolCallReference->InstanceObject = TStrongObjectPtr(InstanceObject);
+				}
+			}
+			return ToolCallReference;
+		}
 		
 		// Ensure a tool call result has no error.
-		bool TestHasNoError(
-			FAutomationTestBase& Test, FObjectFunctionToolCall::FToolCallResult& Result)
+		bool TestHasNoError(FAutomationTestBase& Test, FJsonValueOrError& Result)
 		{
 			const FString* Error = Result.TryGetError();
 			return !Error ||
@@ -50,36 +98,35 @@ namespace UE::ToolsetRegistry
 
 		// Ensure a tool call result has a valid value.
 		TSharedPtr<FJsonObject> TestHasValue(
-			FAutomationTestBase& Test, FObjectFunctionToolCall::FToolCallResult& Result)
+			FAutomationTestBase& Test, FJsonValueOrError& Result)
 		{
 			TSharedPtr<FJsonObject> JsonObject;
 			if (TestHasNoError(Test, Result) &&
 				Test.TestTrue(TEXT("Tool call result has value."), Result.HasValue()))
 			{
-				JsonObject = Result.GetValue();
-				Test.TestTrue(TEXT("Tool call result has valid value."), JsonObject.IsValid());
-#if DEBUG_PRINT_JSONS_TO_LOG
-				if (JsonObject)
+				auto JsonValue = Result.GetValue();
+				if (Test.TestTrue(TEXT("Tool call result has valid value."), JsonValue.IsValid()))
 				{
-					PrintJsonToLog(TEXT("Tool call result value"), JsonObject.ToSharedRef());
-				}
+					JsonObject = JsonValue->AsObject();
+#if DEBUG_PRINT_JSONS_TO_LOG
+					PrintJsonToLog(TEXT("Tool call result value"), JsonValue.ToSharedRef());
 #endif  // DEBUG_PRINT_JSONS_TO_LOG
+				}
 			}
 			return JsonObject;
 		}
 
 		// Ensure a tool call result JSON object has a returnValue field.
-		TSharedPtr<FJsonObject> TestHasReturnValueField(
-			FAutomationTestBase& Test, FObjectFunctionToolCall::FToolCallResult& Result,
-			EJson JsonType)	
+		const TSharedPtr<FJsonObject> TestHasReturnValueField(
+			FAutomationTestBase& Test, FJsonValueOrError& Result, EJson JsonType)	
 		{
-			TSharedPtr<FJsonObject> JsonObject = TestHasValue(Test, Result);
+			const TSharedPtr<FJsonObject> JsonObject = TestHasValue(Test, Result);
 			if (JsonObject &&
 				!Test.TestTrue(
 					TEXT("Tool call result has returnValue."),
 					JsonObject->HasTypedField(TEXT("returnValue"), JsonType)))
 			{
-				JsonObject.Reset();
+				return TSharedPtr<FJsonObject>();
 			}
 			return JsonObject;
 		}
@@ -93,40 +140,30 @@ END_DEFINE_SPEC(FObjectFunctionToolCallTest)
 void FObjectFunctionToolCallTest::Define()
 {
 	using namespace UE::ToolsetRegistry;
+	using namespace UE::ToolsetRegistry::Internal;
 
 	Describe(TEXT("ExecuteToolCall"), [this]()
 	{
 		It(TEXT("should execute a function with a required parameter"), [this]()
 		{
-			static const FName FunctionName = TEXT("TestFuncWithRequiredParam");
 			static constexpr float ClassFloatValue = 20.0f;
 			static constexpr float FunctionInputFloatValue = 3.0f;
-			static const FString FunctionInputParamsJsonString =
-				FString::Printf(TEXT("{\n\"InFloat\": %f\n}"), FunctionInputFloatValue);
-
-			// Create an object and provide it specific internal state. 
-			// (For this test - Using TStrongObjectPtr to prevent GC.)
-
-			const TStrongObjectPtr InstanceObject =
-				TStrongObjectPtr(NewObject<UToolCallFakeClass>());
 
-			InstanceObject->Float = ClassFloatValue;
+			auto ToolCallRef =
+				TestCreateToolCall<UToolCallFakeClass>(*this, TEXT("TestFuncWithRequiredParam"));
+			if (!ToolCallRef.IsSet()) return;
+			ToolCallRef->InstanceObject->Float = ClassFloatValue;
 
 			// Execute.
-
-			const FObjectFunctionToolCall ToolCall =
-				FObjectFunctionToolCall::ConstructWithNamedClassFunction(
-					InstanceObject.Get(), FunctionName);
-
-			FObjectFunctionToolCall::FFunctionInputParamsJson FunctionInputParamsJson;
-			FunctionInputParamsJson.Set<FString>(FunctionInputParamsJsonString);
-
 			AddCommand(
-				UE::ToolsetRegistry::Internal::FUntilFutureCommand<
-						FObjectFunctionToolCall::FToolCallResult>::Create(
-					ToolCall.Execute(FunctionInputParamsJson),
-					[this, InstanceObject](
-						FObjectFunctionToolCall::FToolCallResult&& ToolCallResult) -> void
+				FUntilFutureCommand<FJsonValueOrError>::Create(
+					ToolCallRef->ToolCall->Execute(
+						FObjectFunctionToolCall::FFunctionInputParamsJson(
+							TInPlaceType<FString>(), 
+							FString::Printf(
+								TEXT(R"json({"InFloat": %f})json"), FunctionInputFloatValue))),
+					[this, InstanceObject = ToolCallRef->InstanceObject](
+						FJsonValueOrError&& ToolCallResult) -> void
 					{
 						TSharedPtr<FJsonObject> JsonObject =
 							TestHasReturnValueField(*this, ToolCallResult, EJson::Number);
@@ -141,34 +178,19 @@ void FObjectFunctionToolCallTest::Define()
 
 		It(TEXT("should execute a function with an optional parameter"), [this]()
 		{
-			static const FName FunctionName = TEXT("TestFuncWithOptionalParam");
 			static constexpr float ClassFloatValue = 20.0f;
 
-			// Create an object and provide it specific internal state. 
-			// (For this test - Using TStrongObjectPtr to prevent GC.)
-
-			const TStrongObjectPtr InstanceObject =
-				TStrongObjectPtr(NewObject<UToolCallFakeClass>());
-
-			InstanceObject->Float = ClassFloatValue;
+			auto ToolCallRef =
+				TestCreateToolCall<UToolCallFakeClass>(*this, TEXT("TestFuncWithOptionalParam"));
+			if (!ToolCallRef.IsSet()) return;
+			ToolCallRef->InstanceObject->Float = ClassFloatValue;
 
 			// Execute.
-
-			const TOptional<FObjectFunctionToolCall> ToolCall =
-				FObjectFunctionToolCall::ConstructWithNamedClassFunction(
-					InstanceObject.Get(), FunctionName);
-			if (!TestTrue(TEXT("Tool call created from class and function name."),
-				ToolCall.IsSet()))
-			{
-				return;
-			}
-
 			AddCommand(
-				UE::ToolsetRegistry::Internal::FUntilFutureCommand<
-						FObjectFunctionToolCall::FToolCallResult>::Create(
-					ToolCall->Execute(),
-					[this, InstanceObject](
-						FObjectFunctionToolCall::FToolCallResult&& ToolCallResult) -> void
+				FUntilFutureCommand<FJsonValueOrError>::Create(
+					ToolCallRef->ToolCall->Execute(),
+					[this, InstanceObject = ToolCallRef->InstanceObject](
+						FJsonValueOrError&& ToolCallResult) -> void
 					{
 						TSharedPtr<FJsonObject> JsonObject =
 							TestHasReturnValueField(*this, ToolCallResult, EJson::Number);
@@ -181,6 +203,78 @@ void FObjectFunctionToolCallTest::Define()
 							0.0001f);
 					}));
 		});
+
+		It(TEXT("should execute an asynchronous function"), [this]()
+		{
+			auto ToolCallRef =
+				TestCreateToolCall<UToolCallAsyncResultReturnPropertyTest>(
+					*this, TEXT("SlowGreeting"));
+			if (!ToolCallRef.IsSet()) return;
+
+			// Execute.
+			AddCommand(
+				FUntilFutureCommand<FJsonValueOrError>::Create(
+					ToolCallRef->ToolCall->Execute(
+						FObjectFunctionToolCall::FFunctionInputParamsJson(
+							TInPlaceType<FString>(), TEXT(R"json({"Subject": "world"})json"))),
+					[this, InstanceObject = ToolCallRef->InstanceObject](
+						FJsonValueOrError&& ToolCallResult) -> void
+					{
+						TSharedPtr<FJsonObject> JsonObject =
+							TestHasReturnValueField(*this, ToolCallResult, EJson::String);
+						if (!JsonObject) return;
+
+						TestEqual(
+							TEXT("Return value"),
+							JsonObject->GetStringField(TEXT("returnValue")),
+							TEXT("Hello world"));
+					}));
+		});
+
+		It(TEXT("should execute an asynchronous function that returns an error"), [this]()
+		{
+			auto ToolCallRef =
+				TestCreateToolCall<UToolCallAsyncResultReturnPropertyTest>(
+					*this, TEXT("SlowError"));
+			if (!ToolCallRef.IsSet()) return;
+
+			static FString ErrorMessage = TEXT("Failed");
+			// Execute.
+			AddCommand(
+				FUntilFutureCommand<FJsonValueOrError>::Create(
+					ToolCallRef->ToolCall->Execute(
+						FObjectFunctionToolCall::FFunctionInputParamsJson(
+							TInPlaceType<FString>(),
+							FString::Printf(TEXT(R"json({"Error": "%s"})json"), *ErrorMessage))),
+					[this, InstanceObject = ToolCallRef->InstanceObject](
+						FJsonValueOrError&& ToolCallResult) -> void
+					{
+						if (!TestTrue(TEXT("Has Error"), ToolCallResult.HasError())) return;
+						TestEqual(TEXT("Error Matches"), ToolCallResult.GetError(), ErrorMessage);
+					}));
+		});
+
+		It(TEXT("should fail to execute an asynchronous function that returns null"), [this]()
+		{
+			auto ToolCallRef =
+				TestCreateToolCall<UToolCallAsyncResultReturnPropertyTest>(
+					*this, TEXT("Broken"));
+			if (!ToolCallRef.IsSet()) return;
+
+			// Execute.
+			AddCommand(
+				FUntilFutureCommand<FJsonValueOrError>::Create(
+					ToolCallRef->ToolCall->Execute(),
+					[this, InstanceObject = ToolCallRef->InstanceObject](
+						FJsonValueOrError&& ToolCallResult) -> void
+					{
+						if (!TestTrue(TEXT("Has Error"), ToolCallResult.HasError())) return;
+						TestEqual(
+							TEXT("Error Matches"),
+							ToolCallResult.GetError(),
+							TEXT("Broken failed to return a UToolCallAsyncResult instance."));
+					}));
+		});
 	});
 }
 
diff --git a/Engine/Plugins/Experimental/ToolsetRegistry/Source/ToolsetRegistry/Private/Tests/ToolCallAsyncResultTest.h b/Engine/Plugins/Experimental/ToolsetRegistry/Source/ToolsetRegistry/Private/Tests/ToolCallAsyncResultTest.h
--- a/Engine/Plugins/Experimental/ToolsetRegistry/Source/ToolsetRegistry/Private/Tests/ToolCallAsyncResultTest.h
+++ b/Engine/Plugins/Experimental/ToolsetRegistry/Source/ToolsetRegistry/Private/Tests/ToolCallAsyncResultTest.h
@@ -2,6 +2,7 @@
 
 #pragma once
 
+#include "Async/Async.h"
 #include "Containers/Array.h"
 #include "Containers/UnrealString.h"
 #include "Misc/AutomationTest.h"
@@ -12,6 +13,7 @@
 #include "ToolsetRegistry/ToolCallAsyncResult.h"
 #include "ToolsetRegistry/ToolCallAsyncResultFutureHandler.h"
 #include "ToolsetRegistry/ToolCallAsyncResultString.h"
+#include "ToolsetRegistry/ToolCallAsyncResultVoid.h"
 
 #include "ToolCallAsyncResultTest.generated.h"
 
@@ -115,16 +117,47 @@ class UToolCallAsyncResultReturnPropertyTest : public UObject
 	GENERATED_BODY()
 
 public:
+	// Async function that returns a result that resolves to a greeting.
 	UFUNCTION()
-	static UToolCallAsyncResultString* SlowGreeting(
-		const FString& Subject)
+	static UToolCallAsyncResultString* SlowGreeting(const FString& Subject)
 	{
 		TObjectPtr<UToolCallAsyncResultString> Result =
 			NewObject<UToolCallAsyncResultString>();
-		Result->SetValue(FString::Printf(TEXT("Hello %s"), *Subject));
+		AsyncTask(ENamedThreads::GameThread,
+			[
+				Result = TStrongObjectPtr(Result.Get()),
+				Subject = FString(Subject)
+			]() -> void
+			{
+				Result->SetValue(FString::Printf(TEXT("Hello %s"), *Subject));
+			});
 		return Result.Get();
 	}
 
+	// Async function that return a result that resolves to an error.
+	UFUNCTION()
+	static UToolCallAsyncResultVoid* SlowError(const FString& Error)
+	{
+		TObjectPtr<UToolCallAsyncResultVoid> Result =
+			NewObject<UToolCallAsyncResultVoid>();
+		AsyncTask(ENamedThreads::GameThread,
+			[
+				Result = TStrongObjectPtr(Result.Get()),
+				Error = FString(Error)
+			]() -> void
+			{
+				Result->SetError(Error);
+			});
+		return Result.Get();
+	}
+
+	// Broken async function that doesn't return a result instance.
+	UFUNCTION()
+	static UToolCallAsyncResultVoid* Broken()
+	{
+		return nullptr;
+	}
+
 	static UFunction* FindTestFunction()
 	{
 		return StaticClass()->FindFunctionByName(TEXT("SlowGreeting"));
diff --git a/Engine/Plugins/Experimental/ToolsetRegistry/Source/ToolsetRegistry/Private/ToolsetRegistry/BlueprintLibraryToolset.cpp b/Engine/Plugins/Experimental/ToolsetRegistry/Source/ToolsetRegistry/Private/ToolsetRegistry/BlueprintLibraryToolset.cpp
--- a/Engine/Plugins/Experimental/ToolsetRegistry/Source/ToolsetRegistry/Private/ToolsetRegistry/BlueprintLibraryToolset.cpp
+++ b/Engine/Plugins/Experimental/ToolsetRegistry/Source/ToolsetRegistry/Private/ToolsetRegistry/BlueprintLibraryToolset.cpp
@@ -17,6 +17,7 @@
 #include "ToolsetRegistryDelegateHandle.h"
 #include "ToolsetRegistryModule.h"
 #include "ToolsetRegistry/JsonConversion.h"
+#include "ToolsetRegistry/JsonValueOrError.h"
 #include "ToolsetRegistry/ObjectFunctionToolCall.h"
 #include "ToolsetRegistry/ToolsetJson.h"
 #include "ToolsetRegistry/ValueOrErrorFuture.h"
@@ -205,10 +206,10 @@ namespace UE::ToolsetRegistry
 		for (UFunction* Function : ValidTools)
 		{
 			check(Function);
-			const FObjectFunctionToolCall ToolCall =
-				FObjectFunctionToolCall(TNotNull<UObject*>(Toolset.Get()), TNotNull<UFunction*>(Function));
 			Tools.Add(
-				Function->GetName(), MakeShared<FObjectFunctionToolCall>(ToolCall));
+				Function->GetName(),
+				FObjectFunctionToolCall::Create(
+					TNotNull<UObject*>(Toolset.Get()), TNotNull<UFunction*>(Function)));
 		}
 	}
 
@@ -248,9 +249,7 @@ namespace UE::ToolsetRegistry
 				ToolCall = ToolCall,
 				ExceptionHandler = ExceptionHandler,
 				ToolNameCopy = ToolNameCopy
-			](
-				FObjectFunctionToolCall::FToolCallResult&& ToolCallResult) ->
-					TValueOrError<FString, FString>
+			](FJsonValueOrError&& ToolCallResult) -> TValueOrError<FString, FString>
 			{
 				ensureAlways(IsInGameThread());
 				if (ToolCallResult.HasError())
@@ -265,7 +264,7 @@ namespace UE::ToolsetRegistry
 					return MakeError(MoveTemp(ScriptException));
 				}
 				check(ToolCallResult.HasValue());
-				TSharedPtr<FJsonObject> ResultValue = ToolCallResult.StealValue();
+				TSharedPtr<FJsonValue> ResultValue = ToolCallResult.StealValue();
 				if (!ResultValue.IsValid())
 				{
 					return MakeError(
diff --git a/Engine/Plugins/Experimental/ToolsetRegistry/Source/ToolsetRegistry/Private/ToolsetRegistry/ObjectFunctionToolCall.cpp b/Engine/Plugins/Experimental/ToolsetRegistry/Source/ToolsetRegistry/Private/ToolsetRegistry/ObjectFunctionToolCall.cpp
--- a/Engine/Plugins/Experimental/ToolsetRegistry/Source/ToolsetRegistry/Private/ToolsetRegistry/ObjectFunctionToolCall.cpp
+++ b/Engine/Plugins/Experimental/ToolsetRegistry/Source/ToolsetRegistry/Private/ToolsetRegistry/ObjectFunctionToolCall.cpp
@@ -6,14 +6,17 @@
 #include "Engine/Engine.h"
 #include "Engine/World.h"
 #include "Engine/EngineTypes.h"
+#include "Misc/Optional.h"
 #include "Serialization/JsonReader.h"
 #include "Serialization/JsonSerializer.h"
 #include "Templates/UnrealTemplate.h"
 #include "ToolsetJson.h"
 #include "UObject/FieldIterator.h"
 
+#include "ToolsetRegistry/PropertyAccessors.h"
 #include "ToolsetRegistry/JsonConversion.h"
 #include "ToolsetRegistry/ToolCallAsyncResult.h"
+#include "ToolsetRegistry/ToolCallAsyncResultFutureHandler.h"
 #include "ToolsetRegistry/ValueOrErrorFuture.h"
 
 namespace UE::ToolsetRegistry
@@ -264,32 +267,33 @@ namespace UE::ToolsetRegistry
 		check(FunctionSchemaJsonObject);
 	}
 
-	/*static*/ FObjectFunctionToolCall FObjectFunctionToolCall::ConstructWithNamedClassFunction(
+	TSharedPtr<FObjectFunctionToolCall> FObjectFunctionToolCall::Create(
+		TNotNull<UObject*> InInstanceObject, TNotNull<UFunction*> InFunction)
+	{
+		return MakeShared<FObjectFunctionToolCall>(InInstanceObject, InFunction);
+	}
+
+	TSharedPtr<FObjectFunctionToolCall> FObjectFunctionToolCall::Create(
 		TNotNull<UObject*> InInstanceObject, const FName& InFunctionName)
 	{
 		check(!InFunctionName.IsNone());
-
 		UFunction* Function = InInstanceObject->GetClass()->FindFunctionByName(InFunctionName);
 		check(Function);
-
-		return FObjectFunctionToolCall(InInstanceObject, Function);
+		return Create(InInstanceObject, Function);
 	}
 
-	TFuture<FObjectFunctionToolCall::FToolCallResult> FObjectFunctionToolCall::Execute(
+	TFuture<FJsonValueOrError> FObjectFunctionToolCall::Execute(
 		const TOptional<FFunctionInputParamsJson>& FunctionInputParamsJson) const
 	{
-		using FJsonObjectOrErrorFuture = TValueOrErrorFuture<TSharedPtr<FJsonObject>>;
-
 		// Get function input params JSON as a JSON object.
 		TSharedPtr<FJsonObject> FunctionInputParamsJsonObject;
-		TValueOrError<TSharedRef<FJsonObject>, FString> JsonObjectOrError =
+		FJsonValueOrError JsonValueOrError =
 			BuildValidFunctionInputParamsJsonObject(FunctionInputParamsJson);
-		if (JsonObjectOrError.HasError())
+		if (JsonValueOrError.HasError())
 		{
-			auto Error = JsonObjectOrError.StealError();
-			return FJsonObjectOrErrorFuture::MakeError(MoveTemp(Error));
+			return FJsonValueOrErrorFuture::Make(MoveTemp(JsonValueOrError));
 		}
-		FunctionInputParamsJsonObject = JsonObjectOrError.GetValue();
+		FunctionInputParamsJsonObject = JsonValueOrError.GetValue()->AsObject();
 
 		// Allocate the frame memory.
 		uint8* FrameMemory;
@@ -330,7 +334,7 @@ namespace UE::ToolsetRegistry
 				FunctionInputParamsJsonObject.ToSharedRef(), Function, FrameMemory))
 		{
 			// Caller error. Ill-formed incoming function input params JSON. 
-			return FJsonObjectOrErrorFuture::MakeError(FString::Printf(TEXT(
+			return FJsonValueOrErrorFuture::MakeError(FString::Printf(TEXT(
 				"Function \"%s\", "
 				"could not convert incoming function input params Json to a UStruct. "
 				"Incoming function input params Json may be incorrect or contain unsupported "
@@ -398,44 +402,65 @@ namespace UE::ToolsetRegistry
 			(FrameMemory + Function->ReturnValueOffset) : nullptr);
 		Function->Invoke(InstanceObject, StackFrame, ReturnValueAddress);
 
-		// Get the return value, if we have one.
-		TSharedPtr<FJsonValue> ReturnJsonValue = Function->GetReturnProperty() ?
-			ToolsetJson::PropertyToJsonData(Function->GetReturnProperty(), ReturnValueAddress) :
-			MakeShared<FJsonValueNull>();
-
-		FString ErrorMessage;
-		TSharedPtr<FJsonObject> ReturnJsonObject = MakeReturnObject(ReturnJsonValue, ErrorMessage);
+		FProperty* ReturnProperty = Function->GetReturnProperty();
 
-		if (!ErrorMessage.IsEmpty())
+		auto MaybeToolCallAsyncResult =
+			ReturnProperty
+			? PropertyValueAsObject<UToolCallAsyncResult>(ReturnProperty, ReturnValueAddress)
+			: TOptional<TObjectPtr<UToolCallAsyncResult>>();
+		// If the return value is an async result, return a future that completes when the result
+		// completes.
+		if (MaybeToolCallAsyncResult.IsSet())
 		{
-			return FJsonObjectOrErrorFuture::MakeError(MoveTemp(ErrorMessage));
+			if (!(*MaybeToolCallAsyncResult))
+			{
+				return FJsonValueOrErrorFuture::MakeError(
+					FString::Printf(
+						TEXT("%s failed to return a UToolCallAsyncResult instance."),
+						*Function->GetName()));
+			}
+			TStrongObjectPtr<UToolCallAsyncResultFutureHandler>
+				ToolCallAsyncResultFutureHandler =
+					UToolCallAsyncResultFutureHandler::Create(*MaybeToolCallAsyncResult);
+			return ToolCallAsyncResultFutureHandler->GetValueAsJson().Next(
+				[This = SharedThis(this), ToolCallAsyncResultFutureHandler](
+					UE::ToolsetRegistry::FJsonValueOrError&& ValueOrError) -> FJsonValueOrError
+				{
+					return This->MakeResult(MoveTemp(ValueOrError));
+				});
 		}
-		return FJsonObjectOrErrorFuture::MakeValue(MoveTemp(ReturnJsonObject));
+
+		// Get the return value, if we have one.
+		TSharedPtr<FJsonValue> ReturnJsonValue =
+			ReturnProperty
+			? ToolsetJson::PropertyToJsonData(ReturnProperty, ReturnValueAddress)
+			: MakeShared<FJsonValueNull>();
+		return FJsonValueOrErrorFuture::Make(MakeResult(MakeValue(ReturnJsonValue)));
 	}
 
-	TSharedPtr<FJsonObject> FObjectFunctionToolCall::MakeReturnObject(
-		const TSharedPtr<FJsonValue>& InReturnValue, FString& OutErrorMessage) const
+	FJsonValueOrError FObjectFunctionToolCall::MakeResult(FJsonValueOrError&& ValueOrError) const
 	{
+		if (ValueOrError.HasError()) return MoveTemp(ValueOrError);
+		
 		// If the value to return is invalid, return an error message instead of the return object.
-		OutErrorMessage.Empty();
-		if (!InReturnValue.IsValid())
+		TSharedPtr<FJsonValue> JsonValue = ValueOrError.GetValue();
+		if (!JsonValue.IsValid())
 		{
-			OutErrorMessage = FString::Printf(
-				TEXT("Function \"%s\": could not convert function return value to Json value."),
-				*Function->GetName());
-			return nullptr;
+			return MakeError(
+				FString::Printf(
+					TEXT("Function \"%s\": could not convert return value to JSON."),
+					*Function->GetName()));
 		}
 
-		// According to the output schema: all returned values are required to be wrapped in a "returnValue"
-		// object.
+		// According to the output schema: all returned values are required to be wrapped in a
+		// "returnValue" object.
 		TSharedPtr<FJsonObject> ReturnJsonObject = MakeShared<FJsonObject>();
-		ReturnJsonObject->SetField(TEXT("returnValue"), InReturnValue);
-		return ReturnJsonObject;
+		ReturnJsonObject->SetField(TEXT("returnValue"), JsonValue);
+		return MakeValue(MakeShared<FJsonValueObject>(ReturnJsonObject));
 	}
 
-	TValueOrError<TSharedRef<FJsonObject>, FString>
-		FObjectFunctionToolCall::BuildValidFunctionInputParamsJsonObject(
-			const TOptional<FFunctionInputParamsJson>& FunctionInputParamsJson) const
+	FJsonValueOrError FObjectFunctionToolCall::BuildValidFunctionInputParamsJsonObject(
+		const TOptional<FFunctionInputParamsJson>& FunctionInputParamsJson) const
 	{
 		// Get incoming function params as JSON object.
 		const TSharedRef<FJsonObject> FunctionInputParamsJsonObject =
@@ -473,7 +498,7 @@ namespace UE::ToolsetRegistry
 			}
 		}
 
-		return MakeValue(FunctionInputParamsJsonObject);
+		return MakeValue(MakeShared<FJsonValueObject>(FunctionInputParamsJsonObject));
 	}
 }
 
diff --git a/Engine/Plugins/Experimental/ToolsetRegistry/Source/ToolsetRegistry/Private/ToolsetRegistry/ObjectFunctionToolCall.h b/Engine/Plugins/Experimental/ToolsetRegistry/Source/ToolsetRegistry/Private/ToolsetRegistry/ObjectFunctionToolCall.h
--- a/Engine/Plugins/Experimental/ToolsetRegistry/Source/ToolsetRegistry/Private/ToolsetRegistry/ObjectFunctionToolCall.h
+++ b/Engine/Plugins/Experimental/ToolsetRegistry/Source/ToolsetRegistry/Private/ToolsetRegistry/ObjectFunctionToolCall.h
@@ -10,71 +10,69 @@
 #include "Templates/ValueOrError.h"
 #include "UObject/Class.h"
 
+#include "ToolsetRegistry/JsonValueOrError.h"
+
 //
 // UE::ToolsetRegistry::FObjectFunctionToolCall
 //
 
 namespace UE::ToolsetRegistry
 {
-	class FObjectFunctionToolCall
+	class FObjectFunctionToolCall :
+		public TSharedFromThis<FObjectFunctionToolCall, ESPMode::ThreadSafe>
 	{
 	public:
+		// Use Create() instead of constructing directly.
+		FObjectFunctionToolCall(
+			TNotNull<UObject*> InInstanceObject, TNotNull<UFunction*> InFunction);
 
 		/**
 		 * Function input params.
 		 * Holds either a JSON object or a JSON string.
 		 */
 		using FFunctionInputParamsJson = TVariant<TSharedPtr<FJsonObject>, FString>;
-
-		/**
-		 * Tool call output result.
-		 * Either a JSON object or a string describing an error.
-		 */
-		using FToolCallResult = TValueOrError<TSharedPtr<FJsonObject>, FString>;
 		
 		/**
-		 * Constructs with a UFunction on a UObject instance. 
+		 * Create with a UFunction on a UObject instance. 
 		* @param InInstanceObject UObject containing the function.
 		 * @param InFunction UFunction contained in the UObject.
 		 */
-		FObjectFunctionToolCall(
+		static TSharedPtr<FObjectFunctionToolCall> Create(
 			TNotNull<UObject*> InInstanceObject, TNotNull<UFunction*> InFunction);
 		
 		/**
-		 * Constructs with a UFunction on a UObject instance, but using the UObject's UClass 
+		 * Create with a UFunction on a UObject instance, but using the UObject's UClass 
 		 * where the function can be found by name.
 		 * @param InInstanceObject UObject containing the function.
 		 * @param InFunctionName UFunction name contained in the UObject.
 		 * @return A function call object.
 		 */
-		static FObjectFunctionToolCall ConstructWithNamedClassFunction(
+		static TSharedPtr<FObjectFunctionToolCall> Create(
 			TNotNull<UObject*> InInstanceObject, const FName& InFunctionName);
 		
 		/**
 		 * Given a JSON containing input parameters for a tool call, execute the function with those 
 		 * input parameters.
-		 * @param FunctionInputParamsJson The optional JSON object or string containing the function 
-		 *		parameters. If the function this class was initialized with has required params, then
-		 *		this must be supplied with JSON that provides those required params. Otherwise, an
-		 *		error will result. If the function this class was initialized with has no required 
-		 *		params, then this does not need to be supplied.
+		 * @param FunctionInputParamsJson The optional JSON object or string containing the
+		 *   function parameters. If the function this class was initialized with has required
+		 *   params, then this must be supplied with JSON that provides those required params.
+		 *   Otherwise, an error will result. If the function this class was initialized with has
+		 *   no required params, then this does not need to be supplied.
 		 * @return Tool call result future. (See above). When tool result does not have an error -
- 		 *		Returned objects will either be (1) the object itself, or (2) a value named 
- 		 *		'returnValue' inside an object. And this may be null if (1) function does not 
- 		 *		return a value, or (2) if function does return a value and that return value 
- 		 *		is null.
+ 		 *   Returned objects will either be (1) the object itself, or (2) a value named 
+ 		 *   'returnValue' inside an object. And this may be null JSON value if (1) function 
+		 *   does not return a value, or (2) if function does return a value and that return
+		 *   value is null.
 		 */
-		TFuture<FToolCallResult> Execute(
+		TFuture<FJsonValueOrError> Execute(
 			const TOptional<FFunctionInputParamsJson>& FunctionInputParamsJson = 
 				TOptional<FFunctionInputParamsJson>()) const;
 
 	private:
-		TValueOrError<TSharedRef<FJsonObject>, FString> 
-		BuildValidFunctionInputParamsJsonObject(
+		FJsonValueOrError BuildValidFunctionInputParamsJsonObject(
 			const TOptional<FFunctionInputParamsJson>& FunctionInputParamsJson) const;
 
-		TSharedPtr<FJsonObject> MakeReturnObject(
-			const TSharedPtr<FJsonValue>& InReturnValue, FString& OutErrorMessage) const;
+		FJsonValueOrError MakeResult(FJsonValueOrError&& ValueOrError) const;
 
 		TNotNull<UObject*> InstanceObject;
 		TNotNull<UFunction*> Function;
