[AI Assistant] Work around false positive static analyzer problem.

#Synced-CL 51189430

diff --git a/Engine/Plugins/Experimental/ToolsetRegistry/Source/ToolsetRegistry/Private/ToolsetRegistry/ObjectFunctionToolCall.cpp b/Engine/Plugins/Experimental/ToolsetRegistry/Source/ToolsetRegistry/Private/ToolsetRegistry/ObjectFunctionToolCall.cpp
--- a/Engine/Plugins/Experimental/ToolsetRegistry/Source/ToolsetRegistry/Private/ToolsetRegistry/ObjectFunctionToolCall.cpp
+++ b/Engine/Plugins/Experimental/ToolsetRegistry/Source/ToolsetRegistry/Private/ToolsetRegistry/ObjectFunctionToolCall.cpp
@@ -412,7 +412,7 @@ namespace UE::ToolsetRegistry
 		// completes.
 		if (MaybeToolCallAsyncResult.IsSet())
 		{
-			if (!(*MaybeToolCallAsyncResult))
+			if (!MaybeToolCallAsyncResult.GetValue())
 			{
 				return FJsonValueOrErrorFuture::MakeError(
 					FString::Printf(
