[AI Assistant] Work around false positive static analyzer failures.

MSVC's static analyzer incorrectly flags potential null pointer dereferences
in ToolCallAsyncResultTest so massage the code to silence it.

#Synced-CL 51194027

diff --git a/Engine/Plugins/Experimental/ToolsetRegistry/Source/ToolsetRegistry/Private/Tests/ToolCallAsyncResultTest.cpp b/Engine/Plugins/Experimental/ToolsetRegistry/Source/ToolsetRegistry/Private/Tests/ToolCallAsyncResultTest.cpp
--- a/Engine/Plugins/Experimental/ToolsetRegistry/Source/ToolsetRegistry/Private/Tests/ToolCallAsyncResultTest.cpp
+++ b/Engine/Plugins/Experimental/ToolsetRegistry/Source/ToolsetRegistry/Private/Tests/ToolCallAsyncResultTest.cpp
@@ -542,12 +542,21 @@ void FToolCallAsyncResultTestProperty::Define()
 	{
 		It(TEXT("Should match an async result class"), [this]()
 		{
+			// NOTE: The additional null pointer checks in the following conditionals are to
+			// silence false positive null pointer deferences flagged by MSVC's static analyzer.
 			UFunction* Function = UToolCallAsyncResultReturnPropertyTest::FindTestFunction();
-			if (!TestTrue(TEXT("Has test function"), Function != nullptr)) return;
+			if (!TestTrue(TEXT("Has test function"), Function != nullptr) || !Function) return;
 			FProperty* ReturnProperty = Function->GetReturnProperty();
-			if (!TestTrue(TEXT("Has return property"), ReturnProperty != nullptr)) return;
+			if (!TestTrue(TEXT("Has return property"), ReturnProperty != nullptr)
+				|| !ReturnProperty)
+			{
+				return;
+			}
 			UClass* ReturnClass = UToolCallAsyncResult::MatchesProperty(ReturnProperty);
-			if (!TestTrue(TEXT("Is a class"), ReturnClass != nullptr)) return;
+			if (!TestTrue(TEXT("Is a class"), ReturnClass != nullptr) || !ReturnClass)
+			{
+				return;
+			}
 			TestTrue(
 				TEXT("Returns async result class"),
 				ReturnClass->IsChildOf<UToolCallAsyncResult>());
@@ -556,9 +565,9 @@ void FToolCallAsyncResultTestProperty::Define()
 		It(TEXT("Should not match a non-async result class"), [this]()
 		{
 			UFunction* Function = UToolCallAsyncResultReturnPropertyTest::FindTestFunction();
-			if (!TestTrue(TEXT("Has test function"), Function != nullptr)) return;
+			if (!TestTrue(TEXT("Has test function"), Function != nullptr) || !Function) return;
 			FProperty* Argument = Function->FindPropertyByName(TEXT("Subject"));
-			if (!TestTrue(TEXT("Has argument"), Argument != nullptr)) return;
+			if (!TestTrue(TEXT("Has argument"), Argument != nullptr) || !Argument) return;
 			TestTrue(
 				TEXT("Does not match non-async result property"),
 				UToolCallAsyncResult::MatchesProperty(Argument) == nullptr);
