[AI Assistant] Make UToolCallAsyncResultReturnPropertyTest a toolset with configurable delay

I was using UToolCallAsyncResultReturnPropertyTest to try to replicate issues with
async tool call execution. Though I was only able to find a UX issue with our tool call
rendering in the web UI, it was useful to be able to modify the delay associated with these
calls so I've decided to leave the code in place.

#Synced-CL 51229138

diff --git a/Engine/Plugins/Experimental/ToolsetRegistry/Source/ToolsetRegistry/Private/Tests/ObjectFunctionToolCallTest.cpp b/Engine/Plugins/Experimental/ToolsetRegistry/Source/ToolsetRegistry/Private/Tests/ObjectFunctionToolCallTest.cpp
--- a/Engine/Plugins/Experimental/ToolsetRegistry/Source/ToolsetRegistry/Private/Tests/ObjectFunctionToolCallTest.cpp
+++ b/Engine/Plugins/Experimental/ToolsetRegistry/Source/ToolsetRegistry/Private/Tests/ObjectFunctionToolCallTest.cpp
@@ -216,7 +216,8 @@ void FObjectFunctionToolCallTest::Define()
 				FUntilFutureCommand<FJsonValueOrError>::Create(
 					ToolCallRef->ToolCall->Execute(
 						FObjectFunctionToolCall::FFunctionInputParamsJson(
-							TInPlaceType<FString>(), TEXT(R"json({"Subject": "world"})json"))),
+							TInPlaceType<FString>(),
+							TEXT(R"json({"Subject": "world", "DelayInSeconds": 0.1})json"))),
 					[this, InstanceObject = ToolCallRef->InstanceObject](
 						FJsonValueOrError&& ToolCallResult) -> void
 					{
@@ -245,7 +246,9 @@ void FObjectFunctionToolCallTest::Define()
 					ToolCallRef->ToolCall->Execute(
 						FObjectFunctionToolCall::FFunctionInputParamsJson(
 							TInPlaceType<FString>(),
-							FString::Printf(TEXT(R"json({"Error": "%s"})json"), *ErrorMessage))),
+							FString::Printf(
+								TEXT(R"json({"Error": "%s", "DelayInSeconds": 0.1})json"),
+								*ErrorMessage))),
 					[this, InstanceObject = ToolCallRef->InstanceObject](
 						FJsonValueOrError&& ToolCallResult) -> void
 					{
diff --git a/Engine/Plugins/Experimental/ToolsetRegistry/Source/ToolsetRegistry/Private/Tests/ToolCallAsyncResultTest.h b/Engine/Plugins/Experimental/ToolsetRegistry/Source/ToolsetRegistry/Private/Tests/ToolCallAsyncResultTest.h
--- a/Engine/Plugins/Experimental/ToolsetRegistry/Source/ToolsetRegistry/Private/Tests/ToolCallAsyncResultTest.h
+++ b/Engine/Plugins/Experimental/ToolsetRegistry/Source/ToolsetRegistry/Private/Tests/ToolCallAsyncResultTest.h
@@ -5,11 +5,15 @@
 #include "Async/Async.h"
 #include "Containers/Array.h"
 #include "Containers/UnrealString.h"
+#include "GenericPlatform/GenericPlatformProcess.h"
 #include "Misc/AutomationTest.h"
+#include "Templates/Function.h"
+#include "Templates/UnrealTemplate.h"
 #include "UObject/Object.h"
 #include "UObject/ObjectMacros.h"
 #include "UObject/ObjectPtr.h"
 
+#include "BPToolsetRegistry.h"
 #include "ToolsetRegistry/ToolCallAsyncResult.h"
 #include "ToolsetRegistry/ToolCallAsyncResultFutureHandler.h"
 #include "ToolsetRegistry/ToolCallAsyncResultString.h"
@@ -17,6 +21,8 @@
 
 #include "ToolCallAsyncResultTest.generated.h"
 
+#if WITH_DEV_AUTOMATION_TESTS
+
 BEGIN_DEFINE_SPEC(
 	FToolCallAsyncResultBaseSpec,
 	"AI.ToolsetRegistry.ToolCallAsyncResult.Base",
@@ -110,44 +116,42 @@ BEGIN_DEFINE_SPEC(
 		            FToolCallAsyncResultWith##AsyncResultValueShortName##Spec)));                  \
 	}
 
+#endif  // WITH_DEV_AUTOMATION_TESTS
+
 // Used to test matching UFunction that returns UToolCallAsyncResult.
 UCLASS(BlueprintType, Hidden)
-class UToolCallAsyncResultReturnPropertyTest : public UObject
+class UToolCallAsyncResultReturnPropertyTest : public UToolsetDefinition
 {
 	GENERATED_BODY()
 
 public:
 	// Async function that returns a result that resolves to a greeting.
 	UFUNCTION()
-	static UToolCallAsyncResultString* SlowGreeting(const FString& Subject)
+	static UToolCallAsyncResultString* SlowGreeting(const FString& Subject, float DelayInSeconds)
 	{
 		TObjectPtr<UToolCallAsyncResultString> Result =
 			NewObject<UToolCallAsyncResultString>();
-		AsyncTask(ENamedThreads::GameThread,
-			[
-				Result = TStrongObjectPtr(Result.Get()),
-				Subject = FString(Subject)
-			]() -> void
+		RunAfterDelay(
+			[Result = TStrongObjectPtr(Result.Get()), Subject = FString(Subject)]() -> void
 			{
 				Result->SetValue(FString::Printf(TEXT("Hello %s"), *Subject));
-			});
+			},
+			DelayInSeconds);
 		return Result.Get();
 	}
 
 	// Async function that return a result that resolves to an error.
 	UFUNCTION()
-	static UToolCallAsyncResultVoid* SlowError(const FString& Error)
+	static UToolCallAsyncResultVoid* SlowError(const FString& Error, float DelayInSeconds)
 	{
 		TObjectPtr<UToolCallAsyncResultVoid> Result =
 			NewObject<UToolCallAsyncResultVoid>();
-		AsyncTask(ENamedThreads::GameThread,
-			[
-				Result = TStrongObjectPtr(Result.Get()),
-				Error = FString(Error)
-			]() -> void
+		RunAfterDelay(
+			[Result = TStrongObjectPtr(Result.Get()), Error = FString(Error)]() -> void
 			{
 				Result->SetError(Error);
-			});
+			},
+			DelayInSeconds);
 		return Result.Get();
 	}
 
@@ -162,4 +166,26 @@ public:
 	{
 		return StaticClass()->FindFunctionByName(TEXT("SlowGreeting"));
 	}
-};
\ No newline at end of file
+
+private:
+	// Run a callable after a delay.
+	static void RunAfterDelay(TFunction<void()>&& Callable, float DelayInSeconds)
+	{
+		Async(
+			EAsyncExecution::Thread,
+			[DelayInSeconds, Callable = MoveTemp(Callable)]()
+			{
+				// ConditionalSleep polls the condition once before sleeping the thread.
+				bool bPolledOnce = false;
+				FGenericPlatformProcess::ConditionalSleep(
+					[&bPolledOnce]() -> bool
+					{
+						bool bStopPolling = bPolledOnce;
+						bPolledOnce = true;
+						return bStopPolling;
+					},
+					DelayInSeconds);
+				Callable();
+			});
+	}
+};
