[AI Assistant] Support waiting for futures that return no value in tests.

#Synced-CL 51285146

diff --git a/Engine/Plugins/Experimental/ToolsetRegistry/Source/ToolsetRegistry/Private/Tests/BlueprintLibraryToolsetTest.cpp b/Engine/Plugins/Experimental/ToolsetRegistry/Source/ToolsetRegistry/Private/Tests/BlueprintLibraryToolsetTest.cpp
--- a/Engine/Plugins/Experimental/ToolsetRegistry/Source/ToolsetRegistry/Private/Tests/BlueprintLibraryToolsetTest.cpp
+++ b/Engine/Plugins/Experimental/ToolsetRegistry/Source/ToolsetRegistry/Private/Tests/BlueprintLibraryToolsetTest.cpp
@@ -468,8 +468,8 @@ IMPLEMENT_SIMPLE_AUTOMATION_TEST(
 	if (!TestTrue(TEXT("Valid toolset handler returned"), Toolset.IsValid())) return false;
 
 	AddCommand(
-		UE::ToolsetRegistry::Internal::FUntilFutureCommand<
-				TValueOrError<FString, FString>>::Create(
+		UE::ToolsetRegistry::Internal::FUntilFutureCommand::Create<
+				TValueOrError<FString, FString>>(
 			Toolset->ExecuteTool(FakeToolName, FakeToolsetInput),
 			[this](TValueOrError<FString, FString>&& Result) -> void
 			{
diff --git a/Engine/Plugins/Experimental/ToolsetRegistry/Source/ToolsetRegistry/Private/Tests/ObjectFunctionToolCallTest.cpp b/Engine/Plugins/Experimental/ToolsetRegistry/Source/ToolsetRegistry/Private/Tests/ObjectFunctionToolCallTest.cpp
--- a/Engine/Plugins/Experimental/ToolsetRegistry/Source/ToolsetRegistry/Private/Tests/ObjectFunctionToolCallTest.cpp
+++ b/Engine/Plugins/Experimental/ToolsetRegistry/Source/ToolsetRegistry/Private/Tests/ObjectFunctionToolCallTest.cpp
@@ -156,7 +156,7 @@ void FObjectFunctionToolCallTest::Define()
 
 			// Execute.
 			AddCommand(
-				FUntilFutureCommand<FJsonValueOrError>::Create(
+				FUntilFutureCommand::Create<FJsonValueOrError>(
 					ToolCallRef->ToolCall->Execute(
 						FObjectFunctionToolCall::FFunctionInputParamsJson(
 							TInPlaceType<FString>(), 
@@ -187,7 +187,7 @@ void FObjectFunctionToolCallTest::Define()
 
 			// Execute.
 			AddCommand(
-				FUntilFutureCommand<FJsonValueOrError>::Create(
+				FUntilFutureCommand::Create<FJsonValueOrError>(
 					ToolCallRef->ToolCall->Execute(),
 					[this, InstanceObject = ToolCallRef->InstanceObject](
 						FJsonValueOrError&& ToolCallResult) -> void
@@ -213,7 +213,7 @@ void FObjectFunctionToolCallTest::Define()
 
 			// Execute.
 			AddCommand(
-				FUntilFutureCommand<FJsonValueOrError>::Create(
+				FUntilFutureCommand::Create<FJsonValueOrError>(
 					ToolCallRef->ToolCall->Execute(
 						FObjectFunctionToolCall::FFunctionInputParamsJson(
 							TInPlaceType<FString>(),
@@ -242,7 +242,7 @@ void FObjectFunctionToolCallTest::Define()
 			static FString ErrorMessage = TEXT("Failed");
 			// Execute.
 			AddCommand(
-				FUntilFutureCommand<FJsonValueOrError>::Create(
+				FUntilFutureCommand::Create<FJsonValueOrError>(
 					ToolCallRef->ToolCall->Execute(
 						FObjectFunctionToolCall::FFunctionInputParamsJson(
 							TInPlaceType<FString>(),
@@ -266,7 +266,7 @@ void FObjectFunctionToolCallTest::Define()
 
 			// Execute.
 			AddCommand(
-				FUntilFutureCommand<FJsonValueOrError>::Create(
+				FUntilFutureCommand::Create<FJsonValueOrError>(
 					ToolCallRef->ToolCall->Execute(),
 					[this, InstanceObject = ToolCallRef->InstanceObject](
 						FJsonValueOrError&& ToolCallResult) -> void
diff --git a/Engine/Plugins/Experimental/ToolsetRegistry/Source/ToolsetRegistry/Private/Tests/UntilFutureCommand.h b/Engine/Plugins/Experimental/ToolsetRegistry/Source/ToolsetRegistry/Private/Tests/UntilFutureCommand.h
--- a/Engine/Plugins/Experimental/ToolsetRegistry/Source/ToolsetRegistry/Private/Tests/UntilFutureCommand.h
+++ b/Engine/Plugins/Experimental/ToolsetRegistry/Source/ToolsetRegistry/Private/Tests/UntilFutureCommand.h
@@ -2,6 +2,8 @@
 
 #pragma once
 
+#include <concepts>
+
 #include "Async/Future.h"
 #include "Misc/AutomationTest.h"
 #include "Templates/Function.h"
@@ -12,25 +14,24 @@
 namespace UE::ToolsetRegistry::Internal
 {
 	// Latent test command that waits until a future is complete or a timeout elapses.
-	template<typename ResultT>
 	class FUntilFutureCommand
 	{
-	private:
-		struct FState
-		{
-			TFuture<ResultT> WaitOnFuture;
-			TFunction<void(ResultT&&)> OnFutureCompletion;
-		};
-
 	public:
+		// InFutureCompletion is called with ResultT&& when InWaitOnFuture is complete.
+		template<typename ResultT> requires (!std::is_void_v<ResultT>)
 		static FUntilCommand* Create(
 			TFuture<ResultT>&& InWaitOnFuture,
-			TFunction<void(ResultT&&)> InFutureCompletion,
-			float InTimeout = 5.0f)
+			TFunction<void(ResultT&&)>&& InFutureCompletion,
+			float InTimeout = DefaultTimeoutInSeconds)
 		{
+			struct FState
+			{
+				TFuture<ResultT> WaitOnFuture;
+				TFunction<void(ResultT&&)> OnFutureCompletion;
+			};
 			auto State = MakeShared<FState>();
 			State->WaitOnFuture = MoveTemp(InWaitOnFuture);
-			State->OnFutureCompletion = InFutureCompletion;
+			State->OnFutureCompletion = MoveTemp(InFutureCompletion);
 			return new FUntilCommand(
 				[State]() -> bool
 				{
@@ -43,6 +44,36 @@ namespace UE::ToolsetRegistry::Internal
 				},
 				InTimeout);
 		}
+
+		template<typename ResultT> requires std::is_void_v<ResultT>
+		static FUntilCommand* Create(
+			TFuture<void>&& InWaitOnFuture,
+			TFunction<void()>&& InFutureCompletion,
+			float InTimeout = DefaultTimeoutInSeconds)
+		{
+			struct FState
+			{
+				TFuture<void> WaitOnFuture;
+				TFunction<void()> OnFutureCompletion;
+			};
+			auto State = MakeShared<FState>();
+			State->WaitOnFuture = MoveTemp(InWaitOnFuture);
+			State->OnFutureCompletion = MoveTemp(InFutureCompletion);
+			return new FUntilCommand(
+				[State]() -> bool
+				{
+					bool bIsReady = State->WaitOnFuture.IsReady();
+					if (bIsReady)
+					{
+						State->OnFutureCompletion();
+					}
+					return bIsReady;
+				},
+				InTimeout);
+		}
+
+	public:
+		static constexpr float DefaultTimeoutInSeconds = 5.0f;
 	};
 }
 
diff --git a/Engine/Plugins/Experimental/ToolsetRegistry/Source/ToolsetRegistry/Private/Tests/UntilFutureCommandTest.cpp b/Engine/Plugins/Experimental/ToolsetRegistry/Source/ToolsetRegistry/Private/Tests/UntilFutureCommandTest.cpp
--- a/Engine/Plugins/Experimental/ToolsetRegistry/Source/ToolsetRegistry/Private/Tests/UntilFutureCommandTest.cpp
+++ b/Engine/Plugins/Experimental/ToolsetRegistry/Source/ToolsetRegistry/Private/Tests/UntilFutureCommandTest.cpp
@@ -22,7 +22,7 @@ void FUnitFutureCommandSpec::Define()
 
 	Describe(TEXT("FUntilFutureCommand"), [this]()
 	{
-		It(TEXT("should complete when a future completes"), [this]()
+		It(TEXT("should complete with a value when a future completes"), [this]()
 		{
 			TSharedPtr<FString> ExpectedResult = MakeShared<FString>(TEXT("yo"));
 			auto Promise = MakeShared<TPromise<FString>>();
@@ -35,13 +35,41 @@ void FUnitFutureCommandSpec::Define()
 					},
 					0.1f /* wait 100ms */));
 			AddCommand(
-				FUntilFutureCommand<FString>::Create(
+				FUntilFutureCommand::Create<FString>(
 					Promise->GetFuture(),
 					[this, ExpectedResult](FString&& Result) -> void
 					{
 						TestEqual(TEXT("result"), Result, *ExpectedResult);
 					}));
 		});
+
+		It(TEXT("should complete with no value when a future completes"), [this]()
+		{
+			auto Promise = MakeShared<TPromise<void>>();
+			auto Completed = MakeShared<bool>(false);
+
+			AddCommand(
+				new FDelayedFunctionLatentCommand(
+					[Promise]() -> void
+					{
+						Promise->EmplaceValue();
+					},
+					0.1f /* wait 100ms */));
+			AddCommand(
+				FUntilFutureCommand::Create<void>(
+					Promise->GetFuture(),
+					[this, Completed]() -> void
+					{
+						*Completed = true;
+					}));
+			AddCommand(
+				new FFunctionLatentCommand(
+					[this, Completed]() -> bool
+					{
+						TestTrue(TEXT("Completed"), *Completed);
+						return true;
+					}));
+		});
 	});
 }
 
