[AI Assistant] Very simple module to run a C++ method on the main thread

This is factored out of ToolCallAsync as it is useful elsewhere e.g handling tool call
dispatch so they're always executed on the main thread due to lack of thread safety
with UObjects.

#Synced-CL 51342976

diff --git a/Engine/Plugins/Experimental/ToolsetRegistry/Source/ToolsetRegistry/Private/Tests/RunOnMainThreadTest.cpp b/Engine/Plugins/Experimental/ToolsetRegistry/Source/ToolsetRegistry/Private/Tests/RunOnMainThreadTest.cpp
new file mode 100644
--- /dev/null
+++ b/Engine/Plugins/Experimental/ToolsetRegistry/Source/ToolsetRegistry/Private/Tests/RunOnMainThreadTest.cpp
@@ -0,0 +1,159 @@
+// Copyright Epic Games, Inc. All Rights Reserved.
+
+#include "ToolsetRegistry/RunOnMainThread.h"
+
+#include "Async/Async.h"
+#include "Containers/UnrealString.h"
+#include "HAL/PlatformTLS.h"
+#include "Misc/AutomationTest.h"
+#include "Misc/Optional.h"
+#include "Templates/SharedPointer.h"
+#include "Templates/UniquePtr.h"
+#include "Templates/UnrealTemplate.h"
+
+#include "ToolsetRegistryTestFlags.h"
+#include "Tests/UntilFutureCommand.h"
+
+#if WITH_DEV_AUTOMATION_TESTS
+
+BEGIN_DEFINE_SPEC(FRunOnMainThreadTest, "AI.ToolsetRegistry.RunOnMainThread",
+	ToolsetRegistryTest::Flags)
+
+	// Ensure the specified thread ID is set and matches the game thread ID.
+	void TestEqualsGameThreadId(TOptional<uint32>& MaybeThreadId)
+	{
+		if (!TestTrue(TEXT("Thread ID is set"), MaybeThreadId.IsSet())) return;
+		TestEqual(TEXT("Thread ID is game thread"), *MaybeThreadId, GameThreadId);
+	}
+
+	// Ensure the specified future has the expected result.
+	void TestFutureHasExpectedResult(TFuture<TUniquePtr<uint32>>& Future)
+	{
+		if (!TestTrue(TEXT("Future is valid"), Future.IsValid())) return;
+		TestIsExpectedResult(Future.Consume());
+	}
+
+	// Ensure the specified value is the expected result.
+	void TestIsExpectedResult(TUniquePtr<uint32>&& Result)
+	{
+		if (!TestTrue(TEXT("Result is set"), Result.IsValid())) return;
+		TestEqual(
+			TEXT("Future value matches"), *Result,
+			*FRunOnMainThreadTest::CreateExpectedResult());
+	}
+
+private:
+	uint32 GameThreadId;
+
+private:
+	static TSharedPtr<TOptional<uint32>> CreateMaybeThreadId()
+	{
+		return MakeShared<TOptional<uint32>>();
+	}
+
+	static TUniquePtr<uint32> CreateExpectedResult()
+	{
+		return MakeUnique<uint32>(42);
+	}
+
+END_DEFINE_SPEC(FRunOnMainThreadTest)
+
+void FRunOnMainThreadTest::Define()
+{
+	using UE::ToolsetRegistry::Internal::RunOnMainThread;
+	using UE::ToolsetRegistry::Internal::FUntilFutureCommand;
+
+	BeforeEach(
+		[this]() -> void
+		{
+			GameThreadId = FPlatformTLS::GetCurrentThreadId();
+		});
+
+	It(TEXT(
+		"Should run a function that returns nothing on the main thread when called from the main "
+		"thread"),
+		[this]() -> void
+	{
+		auto MaybeThreadId = CreateMaybeThreadId();
+		RunOnMainThread(
+			[MaybeThreadId]() -> void
+			{
+				MaybeThreadId->Emplace(FPlatformTLS::GetCurrentThreadId());
+			}).Wait();
+		TestEqualsGameThreadId(*MaybeThreadId);
+	});
+
+	It(TEXT(
+		"Should run a function that returns nothing on the main thread when called from another "	
+		"thread."),
+		[this]() -> void
+	{
+		auto MaybeThreadId = CreateMaybeThreadId();
+		auto Result = Async(
+			EAsyncExecution::Thread,
+			[MaybeThreadId]() -> void
+			{
+				RunOnMainThread(
+					[MaybeThreadId]() -> void
+					{
+						MaybeThreadId->Emplace(FPlatformTLS::GetCurrentThreadId());
+					}).Wait();
+			});
+		AddCommand(
+			FUntilFutureCommand::Create<void>(
+				MoveTemp(Result),
+				[this, MaybeThreadId]() -> void
+				{
+					TestEqualsGameThreadId(*MaybeThreadId);
+				}));
+	});
+
+	It(TEXT(
+		"Should run a function that returns a value on the main thread when called from the main "
+		"thread"),
+		[this]() -> void
+	{
+		auto MaybeThreadId = CreateMaybeThreadId();
+		// This returns TUniquePtr<uint32> to ensure that the return value is moved.
+		auto Result = RunOnMainThread(
+			[MaybeThreadId]() -> TUniquePtr<uint32>
+			{
+				MaybeThreadId->Emplace(FPlatformTLS::GetCurrentThreadId());
+				return FRunOnMainThreadTest::CreateExpectedResult();
+			});
+		TestEqualsGameThreadId(*MaybeThreadId);
+		TestFutureHasExpectedResult(Result);
+	});
+
+	It(TEXT(
+		"Should run a function that returns a value on the main thread when called from another "
+		"thread."),
+		[this]() -> void
+	{
+		auto MaybeThreadId = CreateMaybeThreadId();
+		// This returns TUniquePtr<uint32> to ensure that the return value is moved.
+		auto Result = Async(
+			EAsyncExecution::Thread,
+			[MaybeThreadId]() -> TUniquePtr<uint32>
+			{
+				return 
+					RunOnMainThread(
+						[MaybeThreadId]() -> TUniquePtr<uint32>
+						{
+							MaybeThreadId->Emplace(FPlatformTLS::GetCurrentThreadId());
+							return FRunOnMainThreadTest::CreateExpectedResult(); 
+						}).Consume();
+			});
+		AddCommand(
+			FUntilFutureCommand::Create<TUniquePtr<uint32>>(
+				MoveTemp(Result),
+				[this, MaybeThreadId](TUniquePtr<uint32>&& ResultValue) -> void
+				{
+					TestEqualsGameThreadId(*MaybeThreadId);
+					TestIsExpectedResult(MoveTemp(ResultValue));
+				}));
+		
+	});
+}
+
+#endif  // WITH_DEV_AUTOMATION_TESTS
\ No newline at end of file
diff --git a/Engine/Plugins/Experimental/ToolsetRegistry/Source/ToolsetRegistry/Private/ToolsetRegistry/RunOnMainThread.h b/Engine/Plugins/Experimental/ToolsetRegistry/Source/ToolsetRegistry/Private/ToolsetRegistry/RunOnMainThread.h
new file mode 100644
--- /dev/null
+++ b/Engine/Plugins/Experimental/ToolsetRegistry/Source/ToolsetRegistry/Private/ToolsetRegistry/RunOnMainThread.h
@@ -0,0 +1,37 @@
+// Copyright Epic Games, Inc. All Rights Reserved.
+
+#pragma once
+
+#include <type_traits>
+
+#include "Async/Async.h"
+#include "Async/Future.h"
+#include "Templates/Function.h"
+#include "Templates/UnrealTemplate.h"
+
+namespace UE::ToolsetRegistry::Internal
+{
+	// Run a function on the main thread as soon as possible.
+	template<typename CallableT>
+	inline auto RunOnMainThread(CallableT&& Callable) -> 
+		TFuture<decltype(Forward<CallableT>(Callable)())>
+	{
+		if (IsInGameThread())
+		{
+			if constexpr (std::is_void_v<decltype(Forward<CallableT>(Callable)())>)
+			{
+				Callable();
+				return MakeFulfilledPromise<void>().GetFuture();
+			}
+			else
+			{
+				return MakeFulfilledPromise<
+					decltype(Forward<CallableT>(Callable)())>(Callable()).GetFuture();
+			}
+		}
+		else
+		{
+			return Async(EAsyncExecution::TaskGraphMainThread, Forward<CallableT>(Callable));
+		}
+	}
+}
\ No newline at end of file
diff --git a/Engine/Plugins/Experimental/ToolsetRegistry/Source/ToolsetRegistry/Private/ToolsetRegistry/ToolCallAsyncResult.cpp b/Engine/Plugins/Experimental/ToolsetRegistry/Source/ToolsetRegistry/Private/ToolsetRegistry/ToolCallAsyncResult.cpp
--- a/Engine/Plugins/Experimental/ToolsetRegistry/Source/ToolsetRegistry/Private/ToolsetRegistry/ToolCallAsyncResult.cpp
+++ b/Engine/Plugins/Experimental/ToolsetRegistry/Source/ToolsetRegistry/Private/ToolsetRegistry/ToolCallAsyncResult.cpp
@@ -2,10 +2,8 @@
 
 #include "ToolsetRegistry/ToolCallAsyncResult.h"
 
-#include "Async/Async.h"
 #include "Async/UniqueLock.h"
 #include "Containers/UnrealString.h"
-#include "CoreGlobals.h"
 #include "Logging/LogMacros.h"
 #include "Misc/AssertionMacros.h"
 #include "Templates/Function.h"
@@ -13,6 +11,7 @@
 
 #include "ToolsetRegistryModule.h"
 #include "ToolsetRegistry/JsonConversion.h"
+#include "ToolsetRegistry/RunOnMainThread.h"
 #include "ToolsetRegistry/ToolCallAsyncResultVoid.h"
 #include "ToolsetRegistry/ToolsetJson.h"
 
@@ -78,7 +77,8 @@ bool UToolCallAsyncResult::BroadcastOnCompletedIfComplete()
 	if (bAlreadyComplete)
 	{
 		TStrongObjectPtr<UToolCallAsyncResult> Result(this);
-		RunOnMainThread([Result = MoveTemp(Result)]() { Result->BroadcastOnCompleted(); });
+		UE::ToolsetRegistry::Internal::RunOnMainThread(
+			[Result = MoveTemp(Result)]() -> void { Result->BroadcastOnCompleted(); });
 	}
 	return bAlreadyComplete;
 }
@@ -94,8 +94,8 @@ bool UToolCallAsyncResult::MaybeBroadcastCompletion(
 		return false;
 	}
 	TStrongObjectPtr<UToolCallAsyncResult> Result(this);
-	RunOnMainThread(
-		[Result = MoveTemp(Result), SetFinalState = MoveTemp(SetFinalState)]()
+	UE::ToolsetRegistry::Internal::RunOnMainThread(
+		[Result = MoveTemp(Result), SetFinalState = MoveTemp(SetFinalState)]() -> void
 		{
 			SetFinalState();
 			Result->BroadcastOnCompleted();
@@ -123,18 +123,6 @@ void UToolCallAsyncResult::BroadcastOnCompleted()
 	OnCompleted.Clear();
 }
 
-void UToolCallAsyncResult::RunOnMainThread(TUniqueFunction<void()>&& Function)
-{
-	if (IsInGameThread())
-	{
-		Function();
-	}
-	else
-	{
-		AsyncTask(ENamedThreads::GameThread, MoveTemp(Function));
-	}
-}
-
 TSharedRef<FJsonObject> UToolCallAsyncResult::GetValueJsonSchema(
 	TSubclassOf<UToolCallAsyncResult> Class)
 {
diff --git a/Engine/Plugins/Experimental/ToolsetRegistry/Source/ToolsetRegistry/Public/ToolsetRegistry/ToolCallAsyncResult.h b/Engine/Plugins/Experimental/ToolsetRegistry/Source/ToolsetRegistry/Public/ToolsetRegistry/ToolCallAsyncResult.h
--- a/Engine/Plugins/Experimental/ToolsetRegistry/Source/ToolsetRegistry/Public/ToolsetRegistry/ToolCallAsyncResult.h
+++ b/Engine/Plugins/Experimental/ToolsetRegistry/Source/ToolsetRegistry/Public/ToolsetRegistry/ToolCallAsyncResult.h
@@ -151,10 +151,6 @@ public:
 	// derives from UToolCallAsyncResult, nullptr otherwise.
 	static UClass* MatchesProperty(TNotNull<const FProperty*> Property);
 
-private:
-	// Run a function on the main thread as soon as possible.
-	static void RunOnMainThread(TUniqueFunction<void()>&& Function);
-
 private:
 	static FString ValuePropertyName;
 };
