[AI Assistant] Improve readability of RunOnMainThread.

#Synced-CL 51357853

diff --git a/Engine/Plugins/Experimental/ToolsetRegistry/Source/ToolsetRegistry/Private/ToolsetRegistry/RunOnMainThread.h b/Engine/Plugins/Experimental/ToolsetRegistry/Source/ToolsetRegistry/Private/ToolsetRegistry/RunOnMainThread.h
--- a/Engine/Plugins/Experimental/ToolsetRegistry/Source/ToolsetRegistry/Private/ToolsetRegistry/RunOnMainThread.h
+++ b/Engine/Plugins/Experimental/ToolsetRegistry/Source/ToolsetRegistry/Private/ToolsetRegistry/RunOnMainThread.h
@@ -18,15 +18,15 @@ namespace UE::ToolsetRegistry::Internal
 	{
 		if (IsInGameThread())
 		{
-			if constexpr (std::is_void_v<decltype(Forward<CallableT>(Callable)())>)
+			using ReturnT = decltype(Forward<CallableT>(Callable)());
+			if constexpr (std::is_void_v<ReturnT>)
 			{
 				Callable();
 				return MakeFulfilledPromise<void>().GetFuture();
 			}
 			else
 			{
-				return MakeFulfilledPromise<
-					decltype(Forward<CallableT>(Callable)())>(Callable()).GetFuture();
+				return MakeFulfilledPromise<ReturnT>(Callable()).GetFuture();
 			}
 		}
 		else
