Expose a generic interface by which ToolsetRegistry can customize the FProperty->Json conversion process. This modularizes the existing conversion code and makes it accessible to external plugins.

Also includes a fix where only refPath of "None" was treated as null.

#Synced-CL 51341874

diff --git a/Engine/Plugins/Experimental/ToolsetRegistry/Source/ToolsetRegistry/Private/Tests/FakeConverter.cpp b/Engine/Plugins/Experimental/ToolsetRegistry/Source/ToolsetRegistry/Private/Tests/FakeConverter.cpp
new file mode 100644
--- /dev/null
+++ b/Engine/Plugins/Experimental/ToolsetRegistry/Source/ToolsetRegistry/Private/Tests/FakeConverter.cpp
@@ -0,0 +1,22 @@
+// Copyright Epic Games, Inc. All Rights Reserved.
+
+#include "ToolsetRegistry/Tests/FakeConverter.h"
+
+#include "Templates/UnrealTemplate.h"
+
+#if WITH_DEV_AUTOMATION_TESTS
+
+namespace UE::ToolsetRegistry
+{
+	FFakeConverter::FFakeConverter() = default;
+
+	FFakeConverter::~FFakeConverter() = default;
+
+	FString FFakeConverter::GetName() const
+	{
+		static FString Name(TEXT("FakeConverter"));
+		return Name;
+	}
+}
+
+#endif
\ No newline at end of file
diff --git a/Engine/Plugins/Experimental/ToolsetRegistry/Source/ToolsetRegistry/Private/Tests/ToolsetJsonTest.cpp b/Engine/Plugins/Experimental/ToolsetRegistry/Source/ToolsetRegistry/Private/Tests/ToolsetJsonTest.cpp
--- a/Engine/Plugins/Experimental/ToolsetRegistry/Source/ToolsetRegistry/Private/Tests/ToolsetJsonTest.cpp
+++ b/Engine/Plugins/Experimental/ToolsetRegistry/Source/ToolsetRegistry/Private/Tests/ToolsetJsonTest.cpp
@@ -242,6 +242,47 @@ void FToolsetJsonSpec::Define()
 		TestEqual("RefCheck", TestIn.TestObject, TestOut.TestObject);
 	});
 
+	It("Ensures refs can be set from all forms of null", [this]()
+	{
+		FToolsetJsonTest Test;
+		TObjectPtr<UToolsetJsonTestObject> TestOther = NewObject<UToolsetJsonTestObject>();
+		UStruct* Struct = FToolsetJsonTest::StaticStruct();
+		FProperty* Property = Struct->FindPropertyByName("TestObject");
+		TSharedPtr<FJsonObject> JsonRef = MakeShared<FJsonObject>();
+
+		// Null Json type.
+		Test.TestObject = TestOther;
+		TestEqual("RefCheck", Test.TestObject, TestOther);
+		JsonRef->SetField(FString(TEXT("refPath")), MakeShared<FJsonValueNull>());
+		TestTrue("Output", JsonDataToProperty(
+			MakeShared<FJsonValueObject>(JsonRef), Property, &Test.TestObject));
+		TestEqual("RefCheck", Test.TestObject, TObjectPtr<UToolsetJsonTestObject>());
+
+		// None is the UE standard.
+		Test.TestObject = TestOther;
+		TestEqual("RefCheck", Test.TestObject, TestOther);
+		JsonRef->SetStringField(FString(TEXT("refPath")), FString(TEXT("None")));
+		TestTrue("Output", JsonDataToProperty(
+			MakeShared<FJsonValueObject>(JsonRef), Property, &Test.TestObject));
+		TestEqual("RefCheck", Test.TestObject, TObjectPtr<UToolsetJsonTestObject>());
+
+		// Empty string also means null.
+		Test.TestObject = TestOther;
+		TestEqual("RefCheck", Test.TestObject, TestOther);
+		JsonRef->SetStringField(FString(TEXT("refPath")), FString(TEXT("")));
+		TestTrue("Output", JsonDataToProperty(
+			MakeShared<FJsonValueObject>(JsonRef), Property, &Test.TestObject));
+		TestEqual("RefCheck", Test.TestObject, TObjectPtr<UToolsetJsonTestObject>());
+
+		// And we also support null just in case.
+		Test.TestObject = TestOther;
+		TestEqual("RefCheck", Test.TestObject, TestOther);
+		JsonRef->SetStringField(FString(TEXT("refPath")), FString(TEXT("null")));
+		TestTrue("Output", JsonDataToProperty(
+			MakeShared<FJsonValueObject>(JsonRef), Property, &Test.TestObject));
+		TestEqual("RefCheck", Test.TestObject, TObjectPtr<UToolsetJsonTestObject>());
+	});
+
 	It("Converts SoftPath property to reference schema", [this]()
 	{
 		UStruct* Struct = FToolsetJsonTest::StaticStruct();
diff --git a/Engine/Plugins/Experimental/ToolsetRegistry/Source/ToolsetRegistry/Private/Tests/ToolsetLibraryTest.cpp b/Engine/Plugins/Experimental/ToolsetRegistry/Source/ToolsetRegistry/Private/Tests/ToolsetLibraryTest.cpp
--- a/Engine/Plugins/Experimental/ToolsetRegistry/Source/ToolsetRegistry/Private/Tests/ToolsetLibraryTest.cpp
+++ b/Engine/Plugins/Experimental/ToolsetRegistry/Source/ToolsetRegistry/Private/Tests/ToolsetLibraryTest.cpp
@@ -62,7 +62,7 @@ void FToolsetLibrarySpec::Define()
 		TestTrue("SetProp", TestObject->VisibleProperty);
 	});
 
-	It("Returns false when setting values on non-existant properties.", [this]()
+	It("Returns false when setting values on non-existent properties.", [this]()
 	{
 		UToolsetJsonTestObject* TestObject = NewObject<UToolsetJsonTestObject>();
 		TSharedPtr<FJsonObject> JsonObject = MakeShared<FJsonObject>();
diff --git a/Engine/Plugins/Experimental/ToolsetRegistry/Source/ToolsetRegistry/Private/Tests/ToolsetRegistryTest.cpp b/Engine/Plugins/Experimental/ToolsetRegistry/Source/ToolsetRegistry/Private/Tests/ToolsetRegistryTest.cpp
--- a/Engine/Plugins/Experimental/ToolsetRegistry/Source/ToolsetRegistry/Private/Tests/ToolsetRegistryTest.cpp
+++ b/Engine/Plugins/Experimental/ToolsetRegistry/Source/ToolsetRegistry/Private/Tests/ToolsetRegistryTest.cpp
@@ -9,6 +9,7 @@
 #include "ToolsetRegistryTestFlags.h"
 #include "ToolsetRegistry/Toolset.h"
 #include "ToolsetRegistry/Tests/FakeToolset.h"
+#include "ToolsetRegistry/Tests/FakeConverter.h"
 
 #if WITH_DEV_AUTOMATION_TESTS
 
@@ -536,6 +537,93 @@ namespace UE::ToolsetRegistry
 		return true;
 	}
 
+	IMPLEMENT_SIMPLE_AUTOMATION_TEST(
+		FToolsetRegistryTestConverterRegistrationSuccess,
+		"AI.ToolsetRegistry.ToolsetConverter.Registration.Success",
+		ToolsetRegistryTest::Flags);
+
+	bool FToolsetRegistryTestConverterRegistrationSuccess::RunTest(
+		const FString& UnusedParameters)
+	{
+		TUniquePtr<FToolsetRegistry> Registry = MakeUnique<FToolsetRegistry>();
+
+		// Test whether a valid toolset converter is accepted.
+		TSharedPtr<FFakeConverter> ValidConverter = MakeShared<FFakeConverter>();
+		(void)TestTrue(
+			TEXT("RegisterConverter_Valid"),
+			Registry->RegisterConverter(ValidConverter));
+
+		return true;
+	}
+	
+	IMPLEMENT_SIMPLE_AUTOMATION_TEST(
+		FToolsetRegistryTestConverterUnregistrationSuccess,
+		"AI.ToolsetRegistry.ToolsetConverter.Unregistration.Success",
+		ToolsetRegistryTest::Flags);
+
+	bool FToolsetRegistryTestConverterUnregistrationSuccess::RunTest(
+		const FString& UnusedParameters)
+	{
+		TUniquePtr<FToolsetRegistry> Registry = MakeUnique<FToolsetRegistry>();
+		TSharedPtr<FFakeConverter> ValidConverter = MakeShared<FFakeConverter>();
+		Registry->RegisterConverter(ValidConverter);
+
+		// Test whether unregistration of a registered toolset handler is accepted.
+		(void)TestTrue(
+			TEXT("UnregisterConverter_Valid"),
+			Registry->UnregisterConverter(ValidConverter));
+
+		return true;
+	}
+
+	IMPLEMENT_SIMPLE_AUTOMATION_TEST(
+		FToolsetRegistryTestConverterGetConverterSuccess,
+		"AI.ToolsetRegistry.ToolsetConverter.GetConverter.Success",
+		ToolsetRegistryTest::Flags);
+
+	bool FToolsetRegistryTestConverterGetConverterSuccess::RunTest(
+		const FString& UnusedParameters)
+	{
+		TUniquePtr<FToolsetRegistry> Registry = MakeUnique<FToolsetRegistry>();
+
+		// Test whether converter lookup succeeds.
+		TSharedPtr<FToolsetJsonConverter> ValidConverter = MakeShared<FFakeConverter>();
+		Registry->RegisterConverter(ValidConverter);
+
+		UStruct* Struct = FFakeConverterTest::StaticStruct();
+		FProperty* Prop = Struct->FindPropertyByName("TestFloat");
+
+		(void)TestEqual(
+			TEXT("GetConverter_Success"),
+			Registry->GetConverterForProperty(Prop), ValidConverter);
+
+		return true;
+	}
+
+	IMPLEMENT_SIMPLE_AUTOMATION_TEST(
+		FToolsetRegistryTestConverterGetConverterFailure,
+		"AI.ToolsetRegistry.ToolsetConverter.GetConverter.Failure",
+		ToolsetRegistryTest::Flags);
+
+	bool FToolsetRegistryTestConverterGetConverterFailure::RunTest(
+		const FString& UnusedParameters)
+	{
+		TUniquePtr<FToolsetRegistry> Registry = MakeUnique<FToolsetRegistry>();
+
+		// Test whether converter lookup fails.
+		TSharedPtr<FToolsetJsonConverter> ValidConverter = MakeShared<FFakeConverter>();
+		Registry->RegisterConverter(ValidConverter);
+
+		UStruct* Struct = FFakeConverterTest::StaticStruct();
+		FProperty* Prop = Struct->FindPropertyByName("TestString");
+
+		(void)TestEqual(
+			TEXT("GetConverter_Failed"),
+			Registry->GetConverterForProperty(Prop), TSharedPtr<FToolsetJsonConverter>());
+
+		return true;
+	}
+
 	BEGIN_DEFINE_SPEC(
 		FToolsetRegistryOnToolsetRegisteredCallbackSpec,
 		"AI.ToolsetRegistry.ToolsetRegistry.OnToolsetRegisteredCallback",
diff --git a/Engine/Plugins/Experimental/ToolsetRegistry/Source/ToolsetRegistry/Private/ToolsetRegistry.cpp b/Engine/Plugins/Experimental/ToolsetRegistry/Source/ToolsetRegistry/Private/ToolsetRegistry.cpp
--- a/Engine/Plugins/Experimental/ToolsetRegistry/Source/ToolsetRegistry/Private/ToolsetRegistry.cpp
+++ b/Engine/Plugins/Experimental/ToolsetRegistry/Source/ToolsetRegistry/Private/ToolsetRegistry.cpp
@@ -12,11 +12,21 @@
 #include "Templates/SharedPointer.h"
 
 #include "ToolsetRegistryModule.h"
+#include "ToolsetRegistry/ColorConverter.h"
+#include "ToolsetRegistry/ReferenceConverter.h"
 #include "ToolsetRegistry/Toolset.h"
 
 namespace UE::ToolsetRegistry
 {
-	FToolsetRegistry::FToolsetRegistry() = default;
+	FToolsetRegistry::FToolsetRegistry()
+	{
+		TSharedPtr<FToolsetColorConverter> ColorConverter = MakeShared<FToolsetColorConverter>();
+		RegisterConverter(ColorConverter);
+
+		TSharedPtr<FToolsetReferenceConverter> ReferenceConverter = MakeShared<FToolsetReferenceConverter>();
+		RegisterConverter(ReferenceConverter);
+	}
+
 	FToolsetRegistry::~FToolsetRegistry() = default;
 
 	TValueOrError<FToolDescriptor, FString> FToolDescriptor::FromString(
@@ -149,4 +159,32 @@ namespace UE::ToolsetRegistry
 		Result += TEXT("]");
 		return Result;
 	}
+
+	bool FToolsetRegistry::RegisterConverter(TSharedPtr<FToolsetJsonConverter> JsonConverter)
+	{
+		verify(JsonConverter.IsValid());
+		verify(!JsonConverters.Contains(JsonConverter->GetName()));
+		JsonConverters.Add(JsonConverter->GetName(), JsonConverter);
+		return true;
+	}
+
+	bool FToolsetRegistry::UnregisterConverter(TSharedPtr<FToolsetJsonConverter> JsonConverter)
+	{
+		verify(JsonConverter.IsValid());
+		verify(JsonConverters.Remove(JsonConverter->GetName()));
+		return true;
+	}
+
+	TSharedPtr<FToolsetJsonConverter> FToolsetRegistry::GetConverterForProperty(
+		TNotNull<const FProperty*> Property)
+	{
+		for (const auto& Converter : JsonConverters)
+		{
+			if (Converter.Value->CanConvertProperty(Property))
+			{
+				return Converter.Value;
+			}
+		}
+		return nullptr;
+	}
 }
diff --git a/Engine/Plugins/Experimental/ToolsetRegistry/Source/ToolsetRegistry/Private/ToolsetRegistry/ColorConverter.cpp b/Engine/Plugins/Experimental/ToolsetRegistry/Source/ToolsetRegistry/Private/ToolsetRegistry/ColorConverter.cpp
new file mode 100644
--- /dev/null
+++ b/Engine/Plugins/Experimental/ToolsetRegistry/Source/ToolsetRegistry/Private/ToolsetRegistry/ColorConverter.cpp
@@ -0,0 +1,138 @@
+// Copyright Epic Games, Inc. All Rights Reserved.
+
+#include "ColorConverter.h"
+
+#include "JsonObjectConverter.h"
+#include "JsonSchema/JsonSchemaGenerator.h"
+#include "Kismet/KismetSystemLibrary.h"
+
+#include "ToolsetRegistry/JsonConversion.h"
+
+namespace
+{
+	const FStructProperty* AsColorProperty(TNotNull<const FProperty*> Property)
+	{
+		if (const FStructProperty* StructProperty = CastField<FStructProperty>(Property))
+		{
+			const FName StructName = StructProperty->Struct->GetFName();
+			if (StructName == NAME_Color || StructName == NAME_LinearColor)
+			{
+				return StructProperty;
+			}
+		}
+		return nullptr;
+	}
+
+	bool IsLinear(const FProperty* Property)
+	{
+		const FStructProperty* ColorProperty = AsColorProperty(Property);
+		check(ColorProperty);
+		return ColorProperty->Struct->GetFName() == NAME_LinearColor;
+	}
+}
+
+namespace UE::ToolsetRegistry
+{
+	FString FToolsetColorConverter::GetName() const
+	{
+		static FString Name(TEXT("ColorConverter"));
+		return Name;
+	}
+
+	bool FToolsetColorConverter::CanConvertProperty(TNotNull<const FProperty*> Property)
+	{
+		return AsColorProperty(Property) != nullptr;
+	}
+
+	TSharedPtr<FJsonObject> FToolsetColorConverter::PropertyToJsonSchema(
+		TNotNull<const FProperty*> Property)
+	{
+		bool IsLinearProperty = IsLinear(Property);
+		TSharedPtr<FJsonObject> Schema = FJsonSchemaGenerator::UStructToJsonSchemaObject(
+			TBaseStructure<FLinearColor>::Get());
+		TSharedPtr<FJsonObject> Properties = Schema->GetObjectField(TEXT("properties"));
+		TArray<FString> PropertyNames({ "r", "g", "b", "a" });
+		for (const auto& Name : PropertyNames)
+		{
+			TSharedPtr<FJsonObject> Component = Properties->GetObjectField(Name);
+			Component->SetNumberField(FString(TEXT("minimum")), 0.f);
+			if (!IsLinearProperty)
+			{
+				Component->SetNumberField(FString(TEXT("maximum")), 1.f);
+			}
+		}
+		// The default description isn't very useful and includes a source code reference,
+		// so we just remove it to save some tokens.
+		Schema->RemoveField(FString(TEXT("description")));
+		return Schema;
+	}
+
+	TSharedPtr<FJsonValue> FToolsetColorConverter::PropertyToDefault(
+		TNotNull<const FProperty*> Property, const FString& DefaultString)
+	{
+		FLinearColor DefaultLinearColor = FLinearColor::Black;
+		if (IsLinear(Property))
+		{
+			Property->ImportText_Direct(*DefaultString, &DefaultLinearColor, nullptr, PPF_None);
+		}
+		else
+		{
+			FColor DefaultColor = FColor::Black;
+			Property->ImportText_Direct(*DefaultString, &DefaultColor, nullptr, PPF_None);
+			DefaultLinearColor = DefaultColor.ReinterpretAsLinear();
+		}
+		TSharedPtr<FJsonObject> DefaultObject = MakeShared<FJsonObject>();
+		FJsonObjectConverter::UStructToJsonObject(
+			TBaseStructure<FLinearColor>::Get(), &DefaultLinearColor, DefaultObject.ToSharedRef());
+		return MakeShared<FJsonValueObject>(DefaultObject);
+	}
+
+	TSharedPtr<FJsonValue> FToolsetColorConverter::PropertyToJsonData(
+		TNotNull<FProperty*> Property, const void* Value)
+	{
+		FLinearColor OutColor;
+		if (IsLinear(Property))
+		{
+			OutColor = *(static_cast<const FLinearColor*>(Value));
+		}
+		else
+		{
+			const FColor& InColor = *(static_cast<const FColor*>(Value));
+			OutColor = InColor.ReinterpretAsLinear();
+		}
+		TSharedPtr<FJsonObject> OutJson = MakeShared<FJsonObject>();
+		FJsonObjectConverter::UStructToJsonObject(
+			TBaseStructure<FLinearColor>::Get(), &OutColor, OutJson.ToSharedRef());
+		return MakeShared<FJsonValueObject>(OutJson);
+	}
+
+	bool FToolsetColorConverter::JsonDataToProperty(
+		const TSharedPtr<FJsonValue>& JsonValue, TNotNull<FProperty*> Property,
+		void* OutValue)
+	{
+		TSharedPtr<FJsonObject> JsonObject = JsonValue->AsObject();
+		
+		FLinearColor LinearColor;
+		if (FJsonObjectConverter::JsonObjectToUStruct(
+			JsonObject.ToSharedRef(), TBaseStructure<FLinearColor>::Get(), &LinearColor))
+		{
+			if (IsLinear(Property))
+			{
+				FLinearColor& LinearColorOut = *(static_cast<FLinearColor*>(OutValue));
+				LinearColorOut = LinearColor;
+			}
+			else
+			{
+				FColor& ColorOut = *(static_cast<FColor*>(OutValue));
+				ColorOut = LinearColor.ToFColor(false);
+			}
+		}
+		else
+		{
+			UKismetSystemLibrary::RaiseScriptError(FString::Printf(
+				TEXT("%s is not valid JSON for Property %s"),
+				*Internal::JsonToString(JsonValue.ToSharedRef()), *Property->GetName()));
+		}
+		return true;
+	}
+}
\ No newline at end of file
diff --git a/Engine/Plugins/Experimental/ToolsetRegistry/Source/ToolsetRegistry/Private/ToolsetRegistry/ColorConverter.h b/Engine/Plugins/Experimental/ToolsetRegistry/Source/ToolsetRegistry/Private/ToolsetRegistry/ColorConverter.h
new file mode 100644
--- /dev/null
+++ b/Engine/Plugins/Experimental/ToolsetRegistry/Source/ToolsetRegistry/Private/ToolsetRegistry/ColorConverter.h
@@ -0,0 +1,29 @@
+// Copyright Epic Games, Inc. All Rights Reserved.
+
+#include "ToolsetRegistry/ToolsetJsonConverter.h"
+
+namespace UE::ToolsetRegistry
+{
+	/// Converts FColor and FLinearColor.
+	class FToolsetColorConverter : public FToolsetJsonConverter
+	{
+	public:
+		virtual FString GetName() const override;
+
+		virtual bool CanConvertProperty(TNotNull<const FProperty*> Property)  override;
+
+		virtual TSharedPtr<FJsonObject> PropertyToJsonSchema(
+			TNotNull<const FProperty*> Property)  override;
+
+		virtual TSharedPtr<FJsonValue> PropertyToDefault(
+			TNotNull<const FProperty*> Property, const FString& DefaultString)  override;
+
+		virtual TSharedPtr<FJsonValue> PropertyToJsonData(
+			TNotNull<FProperty*> Property, const void* Value)  override;
+
+		virtual bool JsonDataToProperty(
+			const TSharedPtr<FJsonValue>& JsonValue, TNotNull<FProperty*> Property,
+			void* OutValue)  override;
+	};
+}
+
diff --git a/Engine/Plugins/Experimental/ToolsetRegistry/Source/ToolsetRegistry/Private/ToolsetRegistry/ReferenceConverter.cpp b/Engine/Plugins/Experimental/ToolsetRegistry/Source/ToolsetRegistry/Private/ToolsetRegistry/ReferenceConverter.cpp
new file mode 100644
--- /dev/null
+++ b/Engine/Plugins/Experimental/ToolsetRegistry/Source/ToolsetRegistry/Private/ToolsetRegistry/ReferenceConverter.cpp
@@ -0,0 +1,183 @@
+// Copyright Epic Games, Inc. All Rights Reserved.
+
+#include "ReferenceConverter.h"
+
+#include "JsonObjectConverter.h"
+#include "JsonSchema/JsonSchemaGenerator.h"
+#include "Kismet/KismetSystemLibrary.h"
+#include "Misc/StringOutputDevice.h"
+
+#include "ToolsetRegistry/PropertyAccessors.h"
+#include "ToolsetRegistry/JsonConversion.h"
+
+namespace
+{
+	bool StringMeansNull(const FString& String)
+	{
+		return String.IsEmpty() || String.Equals("None", ESearchCase::IgnoreCase) || String.Equals("null", ESearchCase::IgnoreCase);
+	}
+}
+
+namespace UE::ToolsetRegistry
+{
+	FString FToolsetReferenceConverter::GetName() const
+	{
+		static FString Name(TEXT("ReferenceConverter"));
+		return Name;
+	}
+
+	bool FToolsetReferenceConverter::CanConvertProperty(TNotNull<const FProperty*> Property)
+	{
+		if (CastField<FObjectProperty>(Property) &&
+			!Property->HasAnyPropertyFlags(CPF_PersistentInstance))
+		{
+			return true;
+		}
+		else if (
+			CastField<FSoftObjectProperty>(Property) ||
+			CastField<FWeakObjectProperty>(Property))
+		{
+			return true;
+		}
+		else if (const FStructProperty* StructProperty = CastField<FStructProperty>(Property))
+		{
+			const UScriptStruct* ScriptStruct = StructProperty->Struct;
+			if (ScriptStruct == FSoftObjectPath::StaticStruct() ||
+				ScriptStruct == FSoftClassPath::StaticStruct())
+			{
+				return true;
+			}
+		}
+
+		return false;
+	}
+
+	TSharedPtr<FJsonObject> FToolsetReferenceConverter::PropertyToJsonSchema(
+		TNotNull<const FProperty*> Property)
+	{
+		TSharedPtr<FJsonObject> Schema = MakeShared< FJsonObject>();
+		FJsonObject::Duplicate(
+			FJsonSchemaGenerator::UStructToJsonSchemaObject<FToolsetReference>(), Schema);
+		UClass* ClassType = nullptr;
+		bool IsClass = false;
+		if (const FClassProperty* ClassProperty = CastField<FClassProperty>(Property))
+		{
+			ClassType = ClassProperty->MetaClass;
+			IsClass = true;
+		}
+		else if (const FObjectPropertyBase* ObjectProperty = CastField<FObjectPropertyBase>(Property))
+		{
+			ClassType = ObjectProperty->PropertyClass;
+		}
+		else if (const FStructProperty* StructProperty = CastField<FStructProperty>(Property))
+		{
+			const UScriptStruct* ScriptStruct = StructProperty->Struct;
+			if (ScriptStruct == FSoftObjectPath::StaticStruct())
+			{
+				ClassType = UObject::StaticClass();
+			}
+			else if (ScriptStruct == FSoftClassPath::StaticStruct())
+			{
+				ClassType = UObject::StaticClass();
+				IsClass = true;
+			}
+		}
+		else
+		{
+			check(false);
+		}
+		check(ClassType);
+		FString TypeInfo = IsClass ? FString(TEXT("Class@")) : FString();
+		TypeInfo += ClassType->GetPathName();
+		Schema->SetStringField(TEXT("title"), TypeInfo);
+		return Schema;
+	}
+
+	TSharedPtr<FJsonValue> FToolsetReferenceConverter::PropertyToDefault(
+		TNotNull<const FProperty*> Property, const FString& DefaultString)
+	{
+		if (StringMeansNull(DefaultString))
+		{
+			return MakeShared<FJsonValueNull>();
+		}
+		else
+		{
+			FToolsetReference Reference;
+			Reference.RefPath = DefaultString;
+			TSharedPtr<FJsonObject> DefaultObject = MakeShared<FJsonObject>();
+			DefaultObject = FJsonObjectConverter::UStructToJsonObject(Reference);
+			return MakeShared<FJsonValueObject>(DefaultObject);
+		}
+	}
+
+	TSharedPtr<FJsonValue> FToolsetReferenceConverter::PropertyToJsonData(
+		TNotNull<FProperty*> Property, const void* Value)
+	{
+		FToolsetReference Reference;
+		if (auto MaybeObject = Internal::PropertyValueAsObject<UObject>(Property, Value);
+			MaybeObject.IsSet())
+		{
+			if (*MaybeObject)
+			{
+				Reference.RefPath = (*MaybeObject)->GetPathName();
+			}
+			else
+			{
+				return nullptr;
+			}
+		}
+		else if (const FStructProperty* StructProperty = CastField<FStructProperty>(Property))
+		{
+			const UScriptStruct* ScriptStruct = StructProperty->Struct;
+			if (ScriptStruct == FSoftObjectPath::StaticStruct())
+			{
+				const FSoftObjectPath* SoftObjectPath = reinterpret_cast<const FSoftObjectPath*>(Value);
+				Reference.RefPath = SoftObjectPath->ToString();
+			}
+			else if (ScriptStruct == FSoftClassPath::StaticStruct())
+			{
+				const FSoftClassPath* SoftClassPath = reinterpret_cast<const FSoftClassPath*>(Value);
+				Reference.RefPath = SoftClassPath->ToString();
+			}
+			else
+			{
+				check(false);
+			}
+		}
+		else
+		{
+			check(false);
+		}
+		return MakeShared<FJsonValueObject>(FJsonObjectConverter::UStructToJsonObject(Reference));
+	}
+
+	bool FToolsetReferenceConverter::JsonDataToProperty(
+		const TSharedPtr<FJsonValue>& JsonValue, TNotNull<FProperty*> Property,
+		void* OutValue)
+	{
+		FString RefPath;
+		if (JsonValue->Type == EJson::Object)
+		{
+			JsonValue->AsObject()->TryGetStringField(TEXT("refPath"), RefPath);
+		}
+		else if (JsonValue->Type == EJson::String)
+		{
+			RefPath = JsonValue->AsString();
+		}
+
+		if (JsonValue->Type == EJson::Null || StringMeansNull(RefPath))
+		{
+			RefPath = "None";
+		}
+
+		FStringOutputDevice ImportError;
+		Property->ImportText_Direct(*RefPath, OutValue, nullptr, 0, &ImportError);
+		if (!ImportError.IsEmpty())
+		{
+			UKismetSystemLibrary::RaiseScriptError(FString::Printf(
+				TEXT("%s is not valid JSON for Property %s"),
+				*Internal::JsonToString(JsonValue.ToSharedRef()), *Property->GetName()));
+		}
+		return true;
+	}
+}
\ No newline at end of file
diff --git a/Engine/Plugins/Experimental/ToolsetRegistry/Source/ToolsetRegistry/Private/ToolsetRegistry/ReferenceConverter.h b/Engine/Plugins/Experimental/ToolsetRegistry/Source/ToolsetRegistry/Private/ToolsetRegistry/ReferenceConverter.h
new file mode 100644
--- /dev/null
+++ b/Engine/Plugins/Experimental/ToolsetRegistry/Source/ToolsetRegistry/Private/ToolsetRegistry/ReferenceConverter.h
@@ -0,0 +1,42 @@
+// Copyright Epic Games, Inc. All Rights Reserved.
+
+#include "ToolsetRegistry/ToolsetJsonConverter.h"
+
+#include "ReferenceConverter.generated.h"
+
+/// Represents a reference to a UObject or UClass.
+USTRUCT(BlueprintType)
+struct FToolsetReference
+{
+	GENERATED_BODY()
+public:
+	/// The reference stored as a soft path string.
+	UPROPERTY(BlueprintReadWrite, Category = "Toolset")
+	FString RefPath;
+};
+
+namespace UE::ToolsetRegistry
+{
+	/// Converts UObject and UClass references.
+	class FToolsetReferenceConverter : public FToolsetJsonConverter
+	{
+	public:
+		virtual FString GetName() const  override;
+
+		virtual bool CanConvertProperty(TNotNull<const FProperty*> Property)  override;
+
+		virtual TSharedPtr<FJsonObject> PropertyToJsonSchema(
+			TNotNull<const FProperty*> Property)  override;
+
+		virtual TSharedPtr<FJsonValue> PropertyToDefault(
+			TNotNull<const FProperty*> Property, const FString& DefaultString)  override;
+
+		virtual TSharedPtr<FJsonValue> PropertyToJsonData(
+			TNotNull<FProperty*> Property, const void* Value)  override;
+
+		virtual bool JsonDataToProperty(
+			const TSharedPtr<FJsonValue>& JsonValue, TNotNull<FProperty*> Property,
+			void* OutValue)  override;
+	};
+}
+
diff --git a/Engine/Plugins/Experimental/ToolsetRegistry/Source/ToolsetRegistry/Private/ToolsetRegistry/ToolsetJson.cpp b/Engine/Plugins/Experimental/ToolsetRegistry/Source/ToolsetRegistry/Private/ToolsetRegistry/ToolsetJson.cpp
--- a/Engine/Plugins/Experimental/ToolsetRegistry/Source/ToolsetRegistry/Private/ToolsetRegistry/ToolsetJson.cpp
+++ b/Engine/Plugins/Experimental/ToolsetRegistry/Source/ToolsetRegistry/Private/ToolsetRegistry/ToolsetJson.cpp
@@ -9,6 +9,7 @@
 #include "Modules/ModuleManager.h"
 #include "Misc/AssertionMacros.h"
 
+#include "ToolsetRegistrySubsystem.h"
 #include "ToolsetRegistry/PropertyAccessors.h"
 #include "ToolsetRegistry/ToolCallAsyncResult.h"
 
@@ -45,46 +46,6 @@ namespace
 		}
 		return CheckFlags;
 	}
-
-	bool IsReferenceLikeProperty(TNotNull<const FProperty*> Property)
-	{
-		if (CastField<FObjectProperty>(Property) &&
-			!Property->HasAnyPropertyFlags(CPF_PersistentInstance))
-		{
-			return true;
-		}
-		else if (
-			CastField<FSoftObjectProperty>(Property) ||
-			CastField<FWeakObjectProperty>(Property))
-		{
-			return true;
-		}
-		else if (const FStructProperty* StructProperty = CastField<FStructProperty>(Property))
-		{
-			const UScriptStruct* ScriptStruct = StructProperty->Struct;
-			if (ScriptStruct == FSoftObjectPath::StaticStruct() ||
-				ScriptStruct == FSoftClassPath::StaticStruct())
-			{
-				return true;
-			}
-		}
-
-		return false;
-	}
-
-	const FStructProperty* AsColorProperty(
-		TNotNull<const FProperty*> Property, bool AllowLinearColor = false)
-	{
-		if (const FStructProperty* StructProperty = CastField<FStructProperty>(Property))
-		{
-			const FName StructName = StructProperty->Struct->GetFName();
-			if (StructName == NAME_Color || (AllowLinearColor && StructName == NAME_LinearColor))
-			{
-				return StructProperty;
-			}
-		}
-		return nullptr;
-	}
 }
 
 namespace UE::ToolsetRegistry::Internal::ToolsetJson
@@ -97,99 +58,33 @@ bool CustomJsonSchema(
 	const FProperty* Property, const FString& ParameterDefaultString, const TSharedRef<FJsonObject>& OutputSchema)
 {
 	RequiredModules::EnsureLoaded();
+
 	if (UClass* Class = UToolCallAsyncResult::MatchesProperty(Property))
 	{
 		FJsonObject::Duplicate(UToolCallAsyncResult::GetValueJsonSchema(Class), OutputSchema);
 		return true;
 	}
-	else if (IsReferenceLikeProperty(Property))
+
+	auto ToolsetRegistrySubsystem = UToolsetRegistrySubsystem::Get(TEXT("CustomJsonSchema"));
+	if (ToolsetRegistrySubsystem.HasValue())
 	{
-		if (ParameterDefaultString.IsEmpty())
-		{
-			FJsonObject::Duplicate(
-				FJsonSchemaGenerator::UStructToJsonSchemaObject<FToolsetReference>(), OutputSchema);
-			UClass* ClassType = nullptr;
-			bool IsClass = false;
-			if (const FClassProperty* ClassProperty = CastField<FClassProperty>(Property))
-			{
-				ClassType = ClassProperty->MetaClass;
-				IsClass = true;
-			}
-			else if (const FObjectPropertyBase* ObjectProperty = CastField<FObjectPropertyBase>(Property))
-			{
-				ClassType = ObjectProperty->PropertyClass;
-			}
-			else if (const FStructProperty* StructProperty = CastField<FStructProperty>(Property))
-			{
-				const UScriptStruct* ScriptStruct = StructProperty->Struct;
-				if (ScriptStruct == FSoftObjectPath::StaticStruct())
-				{
-					ClassType = UObject::StaticClass();
-				}
-				else if (ScriptStruct == FSoftClassPath::StaticStruct())
-				{
-					ClassType = UObject::StaticClass();
-					IsClass = true;
-				}
-			}
-			check(ClassType);
-			FString TypeInfo = IsClass ? FString(TEXT("Class@")) : FString();
-			TypeInfo += ClassType->GetPathName();
-			OutputSchema->SetStringField(TEXT("title"), TypeInfo);
-		}
-		else
+		auto Converter = ToolsetRegistrySubsystem.GetValue()->ToolsetRegistry.GetConverterForProperty(Property);
+		if (Converter)
 		{
-			if (ParameterDefaultString == FString(TEXT("None")))
+			if (ParameterDefaultString.IsEmpty())
 			{
-				OutputSchema->SetField(TEXT("default"), MakeShared<FJsonValueNull>());
+				TSharedPtr<FJsonObject> PropertyJson = Converter->PropertyToJsonSchema(Property);
+				FJsonObject::Duplicate(PropertyJson, OutputSchema);
 			}
 			else
 			{
-				FToolsetReference Reference;
-				Reference.RefPath = ParameterDefaultString;
-				OutputSchema->SetObjectField(
-					TEXT("default"),
-					FJsonObjectConverter::UStructToJsonObject(Reference));
-			}
-		}
-		return true;
-	}
-	else if (const FStructProperty* ColorProperty = AsColorProperty(Property, true))
-	{
-		bool IsLinear = ColorProperty->Struct->GetFName() == NAME_LinearColor;
-		if (ParameterDefaultString.IsEmpty())
-		{
-			TSharedPtr Schema = FJsonSchemaGenerator::UStructToJsonSchemaObject(
-				TBaseStructure<FLinearColor>::Get());
-			TSharedPtr<FJsonObject> Properties = Schema->GetObjectField(TEXT("properties"));
-			TArray<FString> PropertyNames({ "r", "g", "b", "a" });
-			for (const auto& Name : PropertyNames)
-			{
-				TSharedPtr<FJsonObject> Component = Properties->GetObjectField(Name);
-				Component->SetNumberField(FString(TEXT("minimum")), 0.f);
-				if (!IsLinear)
-				{
-					Component->SetNumberField(FString(TEXT("maximum")), 1.f);
-				}
+				OutputSchema->SetField(
+					TEXT("default"), Converter->PropertyToDefault(Property, ParameterDefaultString));
 			}
-			// The default description isn't very useful and includes a source code reference,
-			// so we just remove it to save some tokens.
-			Schema->RemoveField(FString(TEXT("description")));
-			FJsonObject::Duplicate(Schema, OutputSchema);
-			return true;
-		}
-		else if (!IsLinear)
-		{
-			FColor DefaultColor = FColor::Black;
-			Property->ImportText_Direct(*ParameterDefaultString, &DefaultColor, nullptr, PPF_None);
-			FLinearColor DefaultLinearColor = DefaultColor.ReinterpretAsLinear();
-			TSharedPtr<FJsonObject> OutJson = MakeShared<FJsonObject>();
-			FJsonObjectConverter::UStructToJsonObject(
-				TBaseStructure<FLinearColor>::Get(), &DefaultLinearColor, OutJson.ToSharedRef());
-			OutputSchema->SetObjectField(TEXT("default"), OutJson);
 			return true;
 		}
 	}
+
 	return false;
 }
 
@@ -199,35 +94,14 @@ bool CustomImport(const TSharedPtr<FJsonValue>& JsonValue, FProperty* Property,
 {
 	ensureAlways(JsonValue.IsValid());
 	RequiredModules::EnsureLoaded();
-	if (IsReferenceLikeProperty(Property))
-	{
-		FString RefPath;
-		if (JsonValue->Type == EJson::Object)
-		{						
-			JsonValue->AsObject()->TryGetStringField(TEXT("refPath"), RefPath);
-		}
-		else if (JsonValue->Type == EJson::String)
-		{
-			RefPath = JsonValue->AsString();
-		}
-		if (!RefPath.IsEmpty() && Property->ImportText_Direct(*RefPath, Value, nullptr, 0))
-		{
-			return true;
-		}
-	}
-	else if (AsColorProperty(Property))
+
+	auto ToolsetRegistrySubsystem = UToolsetRegistrySubsystem::Get(TEXT("CustomJsonImport"));
+	if (ToolsetRegistrySubsystem.HasValue())
 	{
-		TSharedPtr<FJsonObject> JsonObject = JsonValue->AsObject();
-		if (JsonObject.IsValid())
+		auto Converter = ToolsetRegistrySubsystem.GetValue()->ToolsetRegistry.GetConverterForProperty(Property);
+		if (Converter)
 		{
-			FLinearColor LinearColor;
-			if (FJsonObjectConverter::JsonObjectToUStruct(
-				JsonObject.ToSharedRef(), TBaseStructure<FLinearColor>::Get(), &LinearColor))
-			{
-				FColor& ColorOut = *(static_cast<FColor*>(Value));
-				ColorOut = LinearColor.ToFColor(false);
-				return true;
-			}
+			return Converter->JsonDataToProperty(JsonValue, Property, Value);
 		}
 	}
 	return false;
@@ -238,58 +112,21 @@ bool CustomImport(const TSharedPtr<FJsonValue>& JsonValue, FProperty* Property,
 TSharedPtr<FJsonValue> CustomExport(FProperty* Property, const void* Value)
 {
 	RequiredModules::EnsureLoaded();
+
 	if (auto Result = PropertyValueAsObject<UToolCallAsyncResult>(Property, Value);
 		Result.IsSet())
 	{
 		return *Result ? (*Result)->GetValueAsJson() : nullptr;
 	}
-	else if (IsReferenceLikeProperty(Property))
+
+	auto ToolsetRegistrySubsystem = UToolsetRegistrySubsystem::Get(TEXT("CustomJsonExport"));
+	if (ToolsetRegistrySubsystem.HasValue())
 	{
-		FToolsetReference Reference;
-		if (auto MaybeObject = PropertyValueAsObject<UObject>(Property, Value);
-			MaybeObject.IsSet())
-		{
-			if (*MaybeObject)
-			{
-				Reference.RefPath = (*MaybeObject)->GetPathName();
-			}
-			else
-			{
-				return nullptr;
-			}
-		}		
-		else if (const FStructProperty* StructProperty = CastField<FStructProperty>(Property))
-		{
-			const UScriptStruct* ScriptStruct = StructProperty->Struct;
-			if (ScriptStruct == FSoftObjectPath::StaticStruct())
-			{
-				const FSoftObjectPath* SoftObjectPath = reinterpret_cast<const FSoftObjectPath*>(Value);
-				Reference.RefPath = SoftObjectPath->ToString();
-			}
-			else if (ScriptStruct == FSoftClassPath::StaticStruct())
-			{
-				const FSoftClassPath* SoftClassPath = reinterpret_cast<const FSoftClassPath*>(Value);
-				Reference.RefPath = SoftClassPath->ToString();
-			}
-			else
-			{
-				check(false);
-			}
-		}
-		else
+		auto Converter = ToolsetRegistrySubsystem.GetValue()->ToolsetRegistry.GetConverterForProperty(Property);
+		if (Converter)
 		{
-			check(false);
+			return Converter->PropertyToJsonData(Property, Value);
 		}
-		return MakeShared<FJsonValueObject>(FJsonObjectConverter::UStructToJsonObject(Reference));
-	}
-	else if (AsColorProperty(Property))
-	{
-		const FColor& InColor = *(static_cast<const FColor*>(Value));
-		FLinearColor OutColor = InColor.ReinterpretAsLinear();
-		TSharedPtr<FJsonObject> OutJson = MakeShared<FJsonObject>();
-		FJsonObjectConverter::UStructToJsonObject(
-			TBaseStructure<FLinearColor>::Get(), &OutColor, OutJson.ToSharedRef());
-		return MakeShared<FJsonValueObject>(OutJson);
 	}
 	
 	return nullptr;
diff --git a/Engine/Plugins/Experimental/ToolsetRegistry/Source/ToolsetRegistry/Private/ToolsetRegistry/ToolsetJson.h b/Engine/Plugins/Experimental/ToolsetRegistry/Source/ToolsetRegistry/Private/ToolsetRegistry/ToolsetJson.h
--- a/Engine/Plugins/Experimental/ToolsetRegistry/Source/ToolsetRegistry/Private/ToolsetRegistry/ToolsetJson.h
+++ b/Engine/Plugins/Experimental/ToolsetRegistry/Source/ToolsetRegistry/Private/ToolsetRegistry/ToolsetJson.h
@@ -5,21 +5,9 @@
 #include "CoreMinimal.h"
 #include "Templates/SharedPointer.h"
 
-#include "ToolsetJson.generated.h"
-
 class FJsonObject;
 class FJsonValue;
 
-/// Represents a reference to a UObject or UClass.
-USTRUCT(BlueprintType)
-struct FToolsetReference
-{
-	GENERATED_BODY()
-public:
-	/// The reference stored as a soft path string.
-	UPROPERTY(BlueprintReadWrite, Category="Toolset")
-	FString RefPath;
-};
 
 /// Contains ToolsetRegistry specific wrappers for JSON schema and data generation.
 /// These wrappers perform custom transformations on specific property types, like
diff --git a/Engine/Plugins/Experimental/ToolsetRegistry/Source/ToolsetRegistry/Private/ToolsetRegistry/ToolsetJsonConverter.cpp b/Engine/Plugins/Experimental/ToolsetRegistry/Source/ToolsetRegistry/Private/ToolsetRegistry/ToolsetJsonConverter.cpp
new file mode 100644
--- /dev/null
+++ b/Engine/Plugins/Experimental/ToolsetRegistry/Source/ToolsetRegistry/Private/ToolsetRegistry/ToolsetJsonConverter.cpp
@@ -0,0 +1,11 @@
+// Copyright Epic Games, Inc. All Rights Reserved.
+
+#include "ToolsetRegistry/ToolsetJsonConverter.h"
+
+namespace UE::ToolsetRegistry
+{
+	FToolsetJsonConverter::FToolsetJsonConverter() = default;
+
+	FToolsetJsonConverter::~FToolsetJsonConverter() = default;
+
+}  // namespace UE::ToolsetRegistry
\ No newline at end of file
diff --git a/Engine/Plugins/Experimental/ToolsetRegistry/Source/ToolsetRegistry/Public/ToolsetRegistry.h b/Engine/Plugins/Experimental/ToolsetRegistry/Source/ToolsetRegistry/Public/ToolsetRegistry.h
--- a/Engine/Plugins/Experimental/ToolsetRegistry/Source/ToolsetRegistry/Public/ToolsetRegistry.h
+++ b/Engine/Plugins/Experimental/ToolsetRegistry/Source/ToolsetRegistry/Public/ToolsetRegistry.h
@@ -9,6 +9,7 @@
 #include "Templates/SharedPointer.h"
 
 #include "ToolsetRegistry/Toolset.h"
+#include "ToolsetRegistry/ToolsetJsonConverter.h"
 
 #define UE_API TOOLSETREGISTRY_API
 
@@ -69,10 +70,23 @@ namespace UE::ToolsetRegistry
 			return OnToolsetRegistryChanged;
 		}
 
+		// Register a custom JSON converter. Returns true if registration was successful.
+		UE_API bool RegisterConverter(TSharedPtr<FToolsetJsonConverter> JsonConverter);
+
+		// Unregister a custom JSON converter. Returns true if unregistration was successful.
+		UE_API bool UnregisterConverter(TSharedPtr<FToolsetJsonConverter> JsonConverter);
+
+		// Returns the converter that can handle the specified property, if any.
+		UE_API TSharedPtr<FToolsetJsonConverter> GetConverterForProperty(
+			TNotNull<const FProperty*> Property);
+
 	private:
 		// The set of registered toolset handlers.
 		TMap<FString, TSharedPtr<FToolset>> ToolsetHandlers;
 
+		// The set of registered converters.
+		TMap<FString, TSharedPtr<FToolsetJsonConverter>> JsonConverters;
+
 		// Callbacks to be called when the toolset registry changes.
 		FOnToolsetRegistryChanged OnToolsetRegistryChanged;
 	};
diff --git a/Engine/Plugins/Experimental/ToolsetRegistry/Source/ToolsetRegistry/Public/ToolsetRegistry/Tests/FakeConverter.h b/Engine/Plugins/Experimental/ToolsetRegistry/Source/ToolsetRegistry/Public/ToolsetRegistry/Tests/FakeConverter.h
new file mode 100644
--- /dev/null
+++ b/Engine/Plugins/Experimental/ToolsetRegistry/Source/ToolsetRegistry/Public/ToolsetRegistry/Tests/FakeConverter.h
@@ -0,0 +1,71 @@
+// Copyright Epic Games, Inc. All Rights Reserved.
+#pragma once
+
+#include "ToolsetRegistry/ToolsetJsonConverter.h"
+
+#include "FakeConverter.generated.h"
+
+USTRUCT(BlueprintType)
+struct FFakeConverterTest
+{
+	GENERATED_BODY()
+public:
+	UPROPERTY(BlueprintReadWrite, Category = "FakeConverter")
+	float TestFloat;
+
+	UPROPERTY(BlueprintReadWrite, Category = "FakeConverter")
+	FString TestString;
+};
+
+#if WITH_DEV_AUTOMATION_TESTS
+
+#define UE_API TOOLSETREGISTRY_API
+
+namespace UE::ToolsetRegistry
+{
+
+class FFakeConverter : public FToolsetJsonConverter
+{
+public:
+	UE_API FFakeConverter();
+	UE_API virtual ~FFakeConverter();
+
+	/// Returns the name of the converter
+	virtual FString GetName() const  override;
+
+	virtual bool CanConvertProperty(TNotNull<const FProperty*> Property)  override
+	{
+		return CastField<FNumericProperty>(Property) != nullptr;
+	}
+
+	virtual TSharedPtr<FJsonObject> PropertyToJsonSchema(
+		TNotNull<const FProperty*> Property)  override
+	{
+		return nullptr;
+	}
+
+	virtual TSharedPtr<FJsonValue> PropertyToDefault(
+		TNotNull<const FProperty*> Property, const FString& DefaultString)  override
+	{
+		return nullptr;
+	}
+
+	virtual TSharedPtr<FJsonValue> PropertyToJsonData(
+		TNotNull<FProperty*> Property, const void* Value)  override
+	{
+		return nullptr;
+	}
+
+	virtual bool JsonDataToProperty(
+		const TSharedPtr<FJsonValue>& JsonValue, TNotNull<FProperty*> Property,
+		void* OutValue)  override
+	{
+		return false;
+	}
+};
+
+}
+
+#undef UE_API
+
+#endif
diff --git a/Engine/Plugins/Experimental/ToolsetRegistry/Source/ToolsetRegistry/Public/ToolsetRegistry/ToolsetJsonConverter.h b/Engine/Plugins/Experimental/ToolsetRegistry/Source/ToolsetRegistry/Public/ToolsetRegistry/ToolsetJsonConverter.h
new file mode 100644
--- /dev/null
+++ b/Engine/Plugins/Experimental/ToolsetRegistry/Source/ToolsetRegistry/Public/ToolsetRegistry/ToolsetJsonConverter.h
@@ -0,0 +1,67 @@
+// Copyright Epic Games, Inc. All Rights Reserved.
+
+#pragma once
+
+#include "Async/Future.h"
+#include "Containers/UnrealString.h"
+#include "Templates/ValueOrError.h"
+
+#define UE_API TOOLSETREGISTRY_API
+
+class FJsonObject;
+class FJsonValue;
+
+namespace UE::ToolsetRegistry
+{
+	/// Base class for custom conversion between FProperty and JSON.
+	class FToolsetJsonConverter
+	{
+	public:
+		UE_API FToolsetJsonConverter();
+		UE_API virtual ~FToolsetJsonConverter();
+
+		/// Returns the name of the converter
+		virtual FString GetName() const = 0;
+
+		/// Determines whether a given property type can be handled by this converter.
+		virtual bool CanConvertProperty(TNotNull<const FProperty*> Property) = 0;
+
+		/*
+		 * Converts an FProperty to JSON schema.
+		 * @param Property The FProperty to convert.
+		 * @return JSON schema that represents the FProperty.
+		 */
+		virtual TSharedPtr<FJsonObject> PropertyToJsonSchema(
+			TNotNull<const FProperty*> Property) = 0;
+
+		/*
+		 * Generates a default object from the FProperty.
+		 * @param Property The FProperty to convert.
+		 * @param DefaultString The default value of the property as a string.
+		 * @return JSON data that represents the FProperty.
+		 */
+		virtual TSharedPtr<FJsonValue> PropertyToDefault(
+			TNotNull<const FProperty*> Property, const FString& DefaultString) = 0;
+
+		/*
+		 * Gets an FProperty on an instance to JSON data.
+		 * @param Property The FProperty to convert.
+		 * @param Value The object or struct to get the value from.
+		 * @return JSON data that represents the value of the FProperty.
+		 */
+		virtual TSharedPtr<FJsonValue> PropertyToJsonData(
+			TNotNull<FProperty*> Property, const void* Value) = 0;
+		/*
+		 * Set an FProperty on an instance from JSON data.
+		 * @param JsonValue The JSON data to convert.
+		 * @param Property The FProperty to apply the data to.
+		 * @param OutValue The object or struct to modify.
+		 * @return True if the data was converted successfully.
+		 */
+		virtual bool JsonDataToProperty(
+			const TSharedPtr<FJsonValue>& JsonValue, TNotNull<FProperty*> Property,
+			void* OutValue) = 0;
+	};
+}
+
+#undef UE_API
\ No newline at end of file
