Add tools that allow the LLM to create and update agent skills.

#Synced-CL 51399443

diff --git a/Engine/Plugins/Experimental/ToolsetRegistry/Source/ToolsetRegistry/Private/Tests/AgentSkillsTest.cpp b/Engine/Plugins/Experimental/ToolsetRegistry/Source/ToolsetRegistry/Private/Tests/AgentSkillsTest.cpp
--- a/Engine/Plugins/Experimental/ToolsetRegistry/Source/ToolsetRegistry/Private/Tests/AgentSkillsTest.cpp
+++ b/Engine/Plugins/Experimental/ToolsetRegistry/Source/ToolsetRegistry/Private/Tests/AgentSkillsTest.cpp
@@ -28,8 +28,9 @@
 BEGIN_DEFINE_SPEC(FAgentSkillSpec, "AI.ToolsetRegistry.AgentSkillSpec",
 	EAutomationTestFlags::EditorContext | EAutomationTestFlags::EngineFilter)
 	FString MountPoint;
-	FString PackagePath;
-	UBlueprint* SkillBlueprint;
+	FString SkillPath;
+	FString DefaultDescription = FString(TEXT("Testy test"));
+	FString DefaultInstructions = FString(TEXT("Do important testy things."));
 	UAgentSkill* AgentSkill;
 	UAgentSkill* PromptSkill;
 END_DEFINE_SPEC(FAgentSkillSpec)
@@ -41,23 +42,20 @@ void FAgentSkillSpec::Define()
 		MountPoint = "/Automation/";
 		FPackageName::RegisterMountPoint(*MountPoint, FPaths::AutomationTransientDir());
 
-		PackagePath = MountPoint + TEXT("TestSkill");
-		UPackage* Package = CreatePackage(*PackagePath);
-		SkillBlueprint = FKismetEditorUtilities::CreateBlueprint(
-			UAgentSkill::StaticClass(), Package, TEXT("TestSkill"), EBlueprintType::BPTYPE_Normal);
+		FAgentSkillDetails Details;
+		Details.Instructions = DefaultInstructions;
+		Details.Toolsets.Add(
+			UE::ToolsetRegistry::FBlueprintLibraryToolset::GetToolsetClassName(
+				UAgentSkillToolset::StaticClass()));
+		SkillPath = UAgentSkillToolset::CreateSkill(
+			MountPoint, FString(TEXT("Test")), DefaultDescription, Details);
 
-		AgentSkill = Cast<UAgentSkill>(SkillBlueprint->GeneratedClass->GetDefaultObject());
-		AgentSkill->Description = "Testy test";
-		AgentSkill->Toolsets.Add(UAgentSkillToolset::StaticClass());
-		AgentSkill->Instructions = "Do important testy things.";
-		Package->MarkPackageDirty();
+		UClass* AgentSkillClass = LoadObject<UClass>(nullptr, SkillPath);
+		AgentSkill = Cast<UAgentSkill>(AgentSkillClass->GetDefaultObject());
 
-		TArray<UPackage*> PackagesToSave;
-		PackagesToSave.Add(Package);
-		FEditorFileUtils::PromptForCheckoutAndSave(PackagesToSave, false, false);
-
-		PromptSkill = Cast<UAgentSkill>(UAgentSkillCustomPrompt::StaticClass()->GetDefaultObject());
-		PromptSkill->Description = FString("CustomPrompt:");
+		PromptSkill = Cast<UAgentSkill>(
+			UAgentSkillCustomPrompt::StaticClass()->GetDefaultObject());
+		PromptSkill->Description = FString(TEXT("CustomPrompt:"));
 	});
 
 	AfterEach([this]()
@@ -65,14 +63,88 @@ void FAgentSkillSpec::Define()
 		PromptSkill->Description.Empty();
 		PromptSkill = nullptr;
 		AgentSkill = nullptr;
+		SkillPath.Empty();
 		GEngine->ForceGarbageCollection(true);
-		UEditorAssetLibrary::DeleteLoadedAssets({ SkillBlueprint });
-		SkillBlueprint = nullptr;
-		PackagePath.Empty();
+		UEditorAssetLibrary::DeleteAsset(FString::Printf(TEXT("%sTest"), *MountPoint));
 		FPackageName::UnRegisterMountPoint(*MountPoint, FPaths::AutomationTransientDir());
 		MountPoint.Empty();
 	});
 
+	It("Can create skills", [this]()
+	{
+		TestEqual(
+			"Created", AgentSkill->Description, DefaultDescription);
+		TestEqual(
+			"Created", AgentSkill->Instructions, DefaultInstructions);
+		TestEqual(
+			"Created", AgentSkill->Toolsets,
+			TArray<TSoftClassPtr<UToolsetDefinition>>({ UAgentSkillToolset::StaticClass() }));
+	});
+
+	It("Won't create when given bad folder paths", [this]()
+	{
+		FAgentSkillDetails Details;
+		TestEqual("BadPath", UAgentSkillToolset::CreateSkill(
+			FString(), FString(), FString(), Details), FString());
+	});
+
+	It("Won't create when given bad asset paths", [this]()
+	{
+		FAgentSkillDetails Details;
+		TestEqual("BadPath", UAgentSkillToolset::CreateSkill(
+			MountPoint, FString(), FString(), Details), FString());
+	});
+
+	It("Won't create when given non-existant folders", [this]()
+	{
+		FAgentSkillDetails Details;
+		TestEqual("BadPath", UAgentSkillToolset::CreateSkill(
+			FString(TEXT("/Bad/Path")), FString(TEXT("Bad")), FString(), Details), FString());
+	});
+
+	It("Can update skills", [this]()
+	{
+		FAgentSkillDetails Details;
+		Details.Instructions = FString(TEXT("Foo bar"));
+		UAgentSkillToolset::UpdateSkill(SkillPath, FString(TEXT("Bish bosh")), Details);
+		TestEqual(
+			"Created", AgentSkill->Description, FString(TEXT("Bish bosh")));
+		TestEqual(
+			"Created", AgentSkill->Instructions, FString(TEXT("Foo bar")));
+		TestEqual(
+			"Created", AgentSkill->Toolsets,
+			TArray<TSoftClassPtr<UToolsetDefinition>>({}));
+	});
+
+	It("Won't update when class does not exist", [this]()
+	{
+		FAgentSkillDetails Details;
+		TestFalse("BadPath", UAgentSkillToolset::UpdateSkill(
+			FString(TEXT("Does/Not/Exist_C")), FString(), Details));
+	});
+
+	It("Won't update when class is not a skill", [this]()
+	{
+		FAgentSkillDetails Details;
+		TestFalse("BadPath", UAgentSkillToolset::UpdateSkill(
+			FString(TEXT("/Script/Engine.StaticMeshActor")), FString(), Details));
+	});
+
+	It("Won't update when given bad tools", [this]()
+	{
+		FAgentSkillDetails Details;
+		Details.Toolsets.Add(FString(TEXT("Does/Not/Exist_C")));
+		TestFalse("BadPath", UAgentSkillToolset::UpdateSkill(
+			SkillPath, FString(), Details));
+		TestEqual(
+			"Created", AgentSkill->Description, DefaultDescription);
+		TestEqual(
+			"Created", AgentSkill->Instructions, DefaultInstructions);
+		TestEqual(
+			"Created", AgentSkill->Toolsets,
+			TArray<TSoftClassPtr<UToolsetDefinition>>({ UAgentSkillToolset::StaticClass() }));
+	});
+
 	It("List skill summaries", [this]()
 	{
 		TMap<FString, FString> Skills = UAgentSkillToolset::ListSkills();
diff --git a/Engine/Plugins/Experimental/ToolsetRegistry/Source/ToolsetRegistry/Private/ToolsetRegistry/AgentSkill.cpp b/Engine/Plugins/Experimental/ToolsetRegistry/Source/ToolsetRegistry/Private/ToolsetRegistry/AgentSkill.cpp
--- a/Engine/Plugins/Experimental/ToolsetRegistry/Source/ToolsetRegistry/Private/ToolsetRegistry/AgentSkill.cpp
+++ b/Engine/Plugins/Experimental/ToolsetRegistry/Source/ToolsetRegistry/Private/ToolsetRegistry/AgentSkill.cpp
@@ -2,12 +2,18 @@
 
 #include "ToolsetRegistry/AgentSkill.h"
 
-#include "AssetRegistry/AssetData.h"
+#include "AssetToolsModule.h"
+#include "Editor.h"
+#include "EditorAssetLibrary.h"
+#include "Factories/BlueprintFactory.h"
 #include "Interfaces/ITargetPlatform.h"
+#include "Kismet/KismetSystemLibrary.h"
 
 #include "ToolsetRegistry/BlueprintLibraryToolset.h"
 #include "ToolsetRegistry/Toolset.h"
 #include "ToolsetLibrary.h"
+#include "ToolsetRegistry/ToolsetRegistry.h"
+#include "ToolsetRegistry/ToolsetRegistrySubsystem.h"
 
 #include UE_INLINE_GENERATED_CPP_BY_NAME(AgentSkill)
 
@@ -58,3 +64,93 @@ TMap<FString, FAgentSkillDetails> UAgentSkillToolset::GetSkills(const TArray<FSt
 	}	
 	return Skills;
 }
+
+FString UAgentSkillToolset::CreateSkill(
+	const FString& FolderPath, const FString& AssetName, const FString& Description,
+	const FAgentSkillDetails& Details)
+{
+	if (AssetName.IsEmpty())
+	{
+		UKismetSystemLibrary::RaiseScriptError(
+			FString(TEXT("Cannot create a skill with no name.")));
+		return FString();
+	}
+
+	if (FolderPath.IsEmpty() || FolderPath[0] != '/')
+	{
+		UKismetSystemLibrary::RaiseScriptError(
+			FString::Printf(TEXT("%s is not a valid folder path."), *FolderPath));
+		return FString();
+	}
+	
+	UBlueprintFactory* BlueprintFactory = NewObject<UBlueprintFactory>();
+	BlueprintFactory->ParentClass = UAgentSkill::StaticClass();
+
+	IAssetTools& AssetTools = FAssetToolsModule::GetModule().Get();
+	UBlueprint* SkillBlueprint = Cast<UBlueprint>(AssetTools.CreateAsset(
+		AssetName, FolderPath, UBlueprint::StaticClass(), BlueprintFactory));
+
+	if (!SkillBlueprint)
+	{
+		UKismetSystemLibrary::RaiseScriptError(
+			FString::Printf(TEXT("Unable to create AgentSkill Blueprint %s in %s."),
+				*AssetName, *FolderPath));
+		return FString();
+	}
+
+	UClass* AgentSkillClass = SkillBlueprint->GeneratedClass;
+	FString AgentClassPath = AgentSkillClass->GetPathName();
+	if (!UpdateSkill(AgentClassPath, Description, Details))
+	{
+		UEditorAssetLibrary::DeleteLoadedAssets({ SkillBlueprint });
+		// No need for an error as it will be populated by UpdateSkill.
+		return FString();
+	}
+
+	return AgentClassPath;
+}
+
+bool UAgentSkillToolset::UpdateSkill(
+	const FString& SkillPath, const FString& Description, const FAgentSkillDetails& Details)
+{
+	UClass* AgentSkillClass = LoadObject<UClass>(nullptr, SkillPath);
+	if (!AgentSkillClass)
+	{
+		UKismetSystemLibrary::RaiseScriptError(
+			FString::Printf(TEXT("%s is not an AgentSkill asset."), *SkillPath));
+		return false;
+	}
+
+	UAgentSkill* AgentSkill = Cast<UAgentSkill>(AgentSkillClass->GetDefaultObject());
+	if (!AgentSkill)
+	{
+		UKismetSystemLibrary::RaiseScriptError(
+			FString::Printf(TEXT("%s is not an AgentSkill asset."), *SkillPath));
+		return false;
+	}
+
+	check(GEditor);
+	UToolsetRegistrySubsystem* ToolsetRegistrySubsystem =
+		GEditor->GetEditorSubsystem<UToolsetRegistrySubsystem>();
+	check(ToolsetRegistrySubsystem);
+	TArray<TSoftClassPtr<UToolsetDefinition>> ToolsetClasses;
+	for (const auto& ToolsetName : Details.Toolsets)
+	{
+		TSharedPtr<UE::ToolsetRegistry::FToolset> Toolset =
+			ToolsetRegistrySubsystem->ToolsetRegistry.Find(ToolsetName);
+		if (!Toolset || !Toolset->GetToolsetClass())
+		{
+			UKismetSystemLibrary::RaiseScriptError(
+				FString::Printf(TEXT("%s is not a valid toolset name."), *ToolsetName));
+			return false;
+		}
+		ToolsetClasses.Add(Toolset->GetToolsetClass());
+	}
+
+	AgentSkill->Description = Description;
+	AgentSkill->Instructions = Details.Instructions;
+	AgentSkill->Toolsets = ToolsetClasses;
+	AgentSkill->Modify();
+	AgentSkillClass->GetPackage()->MarkPackageDirty();
+	return true;
+}
diff --git a/Engine/Plugins/Experimental/ToolsetRegistry/Source/ToolsetRegistry/Private/ToolsetRegistry/BlueprintLibraryToolset.cpp b/Engine/Plugins/Experimental/ToolsetRegistry/Source/ToolsetRegistry/Private/ToolsetRegistry/BlueprintLibraryToolset.cpp
--- a/Engine/Plugins/Experimental/ToolsetRegistry/Source/ToolsetRegistry/Private/ToolsetRegistry/BlueprintLibraryToolset.cpp
+++ b/Engine/Plugins/Experimental/ToolsetRegistry/Source/ToolsetRegistry/Private/ToolsetRegistry/BlueprintLibraryToolset.cpp
@@ -48,7 +48,7 @@ namespace UE::ToolsetRegistry
 
 			~FToolCallScriptExceptionHandler() = default;
 
-			// Track script exceptions raised by UKismetSystemLibrary::RaiseScriptException or
+			// Track script exceptions raised by UKismetSystemLibrary::RaiseScriptError or
 			// FBlueprintCoreDelegates::ThrowScriptException via FBlueprintCoreDelegates::OnScriptException
 			// and copy into local storage.
 			void OnScriptException(
diff --git a/Engine/Plugins/Experimental/ToolsetRegistry/Source/ToolsetRegistry/Public/ToolsetRegistry/AgentSkill.h b/Engine/Plugins/Experimental/ToolsetRegistry/Source/ToolsetRegistry/Public/ToolsetRegistry/AgentSkill.h
--- a/Engine/Plugins/Experimental/ToolsetRegistry/Source/ToolsetRegistry/Public/ToolsetRegistry/AgentSkill.h
+++ b/Engine/Plugins/Experimental/ToolsetRegistry/Source/ToolsetRegistry/Public/ToolsetRegistry/AgentSkill.h
@@ -17,7 +17,7 @@ struct FAgentSkillDetails
 {
 	GENERATED_USTRUCT_BODY()
 public:
-	/// The names of the toolsets the skill depends on.
+	/// The full names of the toolsets the skill depends on.
 	UPROPERTY(BlueprintReadWrite, Category = "AgentSkill")
 	TArray<FString> Toolsets;
 
@@ -56,6 +56,7 @@ protected:
 	FString Instructions;
 
 	friend class FAgentSkillSpec;
+	friend class UAgentSkillToolset;
 };
 
 
@@ -79,4 +80,29 @@ public:
 	 */	
 	UFUNCTION(BlueprintCallable, Category = "Agent Skill")
 	static TMap<FString, FAgentSkillDetails> GetSkills(const TArray<FString>& SkillPaths);
+
+	/**
+	 * Creates a new AgentSkill.
+	 * This should ONLY be called after getting explicit direction or permission from the user.
+	 * @param FolderPath The folder in which to create the skill. i.e. /Game/Skills/.
+	 * @param AssetName The name of the skill to create i.e. MySkill.
+	 * @param Description The brief description of the skill.
+	 * @param Details Detailed information about how to use the skill.
+	 * @return The path for the created Skill class. Empty if unsuccessful.
+	 */
+	UFUNCTION(BlueprintCallable, Category = "Agent Skill")
+	static FString CreateSkill(
+		const FString& FolderPath, const FString& AssetName, const FString& Description, const FAgentSkillDetails& Details);
+
+	/**
+	 * Updates an existing AgentSkill.
+	 * This should ONLY be called after getting explicit direction or permission from the user.
+	 * @param SkillPath The full path to the skill to modify i.e. /Game/Skills/MySkill.MySkill_C.
+	 * @param Description The brief description of the skill.
+	 * @param Details Detailed information about how to use the skill.
+	 * @return True if the skill was updated.
+	 */
+	UFUNCTION(BlueprintCallable, Category = "Agent Skill")
+	static bool UpdateSkill(
+		const FString& SkillPath, const FString& Description, const FAgentSkillDetails& Details);
 };
diff --git a/Engine/Plugins/Experimental/ToolsetRegistry/Source/ToolsetRegistry/Public/ToolsetRegistry/BlueprintLibraryToolset.h b/Engine/Plugins/Experimental/ToolsetRegistry/Source/ToolsetRegistry/Public/ToolsetRegistry/BlueprintLibraryToolset.h
--- a/Engine/Plugins/Experimental/ToolsetRegistry/Source/ToolsetRegistry/Public/ToolsetRegistry/BlueprintLibraryToolset.h
+++ b/Engine/Plugins/Experimental/ToolsetRegistry/Source/ToolsetRegistry/Public/ToolsetRegistry/BlueprintLibraryToolset.h
@@ -48,6 +48,9 @@ namespace UE::ToolsetRegistry
 			return FString("1.0");
 		}
 
+		// Returns the UClass that backs this toolset.
+		virtual UClass* GetToolsetClass() const override { return ToolsetClass.Get(); }
+
 		// Returns the description of the toolset from the class's metadata.
 		FString GetToolsetDescription() const
 		{
diff --git a/Engine/Plugins/Experimental/ToolsetRegistry/Source/ToolsetRegistry/Public/ToolsetRegistry/Toolset.h b/Engine/Plugins/Experimental/ToolsetRegistry/Source/ToolsetRegistry/Public/ToolsetRegistry/Toolset.h
--- a/Engine/Plugins/Experimental/ToolsetRegistry/Source/ToolsetRegistry/Public/ToolsetRegistry/Toolset.h
+++ b/Engine/Plugins/Experimental/ToolsetRegistry/Source/ToolsetRegistry/Public/ToolsetRegistry/Toolset.h
@@ -34,6 +34,9 @@ namespace UE::ToolsetRegistry
 
 		/// Returns the version of the toolset
 		virtual FString GetToolsetVersion() const = 0;
+
+		/// Returns the UClass that implements this toolset, if any.
+		virtual UClass* GetToolsetClass() const { return nullptr; }
 	};
 
 }
