Add UE_API to public UFUNCTIONs in ToolCallAsyncResult headers

#Synced-CL 51401176

diff --git a/Engine/Plugins/Experimental/ToolsetRegistry/Source/ToolsetRegistry/Public/ToolsetRegistry/ToolCallAsyncResult.h b/Engine/Plugins/Experimental/ToolsetRegistry/Source/ToolsetRegistry/Public/ToolsetRegistry/ToolCallAsyncResult.h
--- a/Engine/Plugins/Experimental/ToolsetRegistry/Source/ToolsetRegistry/Public/ToolsetRegistry/ToolCallAsyncResult.h
+++ b/Engine/Plugins/Experimental/ToolsetRegistry/Source/ToolsetRegistry/Public/ToolsetRegistry/ToolCallAsyncResult.h
@@ -68,7 +68,7 @@ public:
 	/// Get the JSON representation of the Value. If the associated tool call is not complete or an
 	/// error occurred, this returns an empty string.
 	UFUNCTION(BlueprintCallable, Category = "ToolsetRegistry")
-	FString GetValueAsJsonString() const;
+	UE_API FString GetValueAsJsonString() const;
 
 	/// Get the JSON representation of the Value. If the associated tool call is not complete or an
 	/// error occurred, this returns an empty (null) shared pointer.
@@ -81,7 +81,7 @@ public:
 	/// @note This is thread safe, the Error property will always be updated and OnCompleted
 	/// will be signaled on the main thread.
 	UFUNCTION(BlueprintCallable, Category = "ToolsetRegistry")
-	bool SetError(const FString& InError);
+	UE_API bool SetError(const FString& InError);
 
 	/// Broadcast to subscribers of the OnCompleted delegate if the result is already complete.
 	///
@@ -95,7 +95,7 @@ public:
 	/// @remark Ideally we would override subscription methods for the OnCompleted delegate
 	/// and notify subscribers there but unfortunately that isn't possible.
 	UFUNCTION(BlueprintCallable, Category = "ToolsetRegistry")
-	bool BroadcastOnCompletedIfComplete();
+	UE_API bool BroadcastOnCompletedIfComplete();
 
 protected:
 	/// Broadcast that the result is complete if it hasn't already been completed.
diff --git a/Engine/Plugins/Experimental/ToolsetRegistry/Source/ToolsetRegistry/Public/ToolsetRegistry/ToolCallAsyncResultString.h b/Engine/Plugins/Experimental/ToolsetRegistry/Source/ToolsetRegistry/Public/ToolsetRegistry/ToolCallAsyncResultString.h
--- a/Engine/Plugins/Experimental/ToolsetRegistry/Source/ToolsetRegistry/Public/ToolsetRegistry/ToolCallAsyncResultString.h
+++ b/Engine/Plugins/Experimental/ToolsetRegistry/Source/ToolsetRegistry/Public/ToolsetRegistry/ToolCallAsyncResultString.h
@@ -6,6 +6,8 @@
 
 #include "ToolCallAsyncResultString.generated.h"
 
+#define UE_API TOOLSETREGISTRY_API
+
 // Async tool call result that completes with a string.
 UCLASS(BlueprintType)
 class UToolCallAsyncResultString : public UToolCallAsyncResult
@@ -15,7 +17,7 @@ class UToolCallAsyncResultString : public UToolCallAsyncResult
 public:
 	// Set the result.
 	UFUNCTION(BlueprintCallable, Category = "ToolsetRegistry")
-	bool SetValue(const FString& InValue)
+	UE_API bool SetValue(const FString& InValue)
 	{
 		return MaybeBroadcastSuccessfulCompletion(FString(InValue), Value);
 	}
@@ -25,3 +27,5 @@ public:
 	FString Value;
 };
 
+#undef UE_API
+
diff --git a/Engine/Plugins/Experimental/ToolsetRegistry/Source/ToolsetRegistry/Public/ToolsetRegistry/ToolCallAsyncResultVoid.h b/Engine/Plugins/Experimental/ToolsetRegistry/Source/ToolsetRegistry/Public/ToolsetRegistry/ToolCallAsyncResultVoid.h
--- a/Engine/Plugins/Experimental/ToolsetRegistry/Source/ToolsetRegistry/Public/ToolsetRegistry/ToolCallAsyncResultVoid.h
+++ b/Engine/Plugins/Experimental/ToolsetRegistry/Source/ToolsetRegistry/Public/ToolsetRegistry/ToolCallAsyncResultVoid.h
@@ -22,7 +22,7 @@ class UToolCallAsyncResultVoid : public UToolCallAsyncResult
 public:
 	// Complete this result with no value and notify listeners of OnCompleted.
 	UFUNCTION(BlueprintCallable, Category = "ToolsetRegistry")
-	bool SetCompleted()
+	UE_API bool SetCompleted()
 	{
 		return MaybeBroadcastSuccessfulCompletion();
 	}
