[Toolset Registry] Improve tool call error messages

#rnx
#Synced-CL 51455374

diff --git a/Engine/Plugins/Experimental/ToolsetRegistry/Source/ToolsetRegistry/Private/ToolsetRegistry/BlueprintLibraryToolset.cpp b/Engine/Plugins/Experimental/ToolsetRegistry/Source/ToolsetRegistry/Private/ToolsetRegistry/BlueprintLibraryToolset.cpp
--- a/Engine/Plugins/Experimental/ToolsetRegistry/Source/ToolsetRegistry/Private/ToolsetRegistry/BlueprintLibraryToolset.cpp
+++ b/Engine/Plugins/Experimental/ToolsetRegistry/Source/ToolsetRegistry/Private/ToolsetRegistry/BlueprintLibraryToolset.cpp
@@ -54,7 +54,7 @@ namespace UE::ToolsetRegistry
 			void OnScriptException(
 				const UObject* ActiveObject, const FFrame& StackFrame, const FBlueprintExceptionInfo& Info)
 			{
-				TOptional<FString> ErrorMessage = GetBlueprintExceptionMessage(Info);
+				TOptional<FString> ErrorMessage = GetBlueprintExceptionMessage(Info, StackFrame);
 				if (ErrorMessage.IsSet())
 				{
 					AddToExceptionMessage(ErrorMessage.GetValue());
@@ -67,7 +67,27 @@ namespace UE::ToolsetRegistry
 			}
 
 		  private:
-			TOptional<FString> GetBlueprintExceptionMessage(const FBlueprintExceptionInfo& Info)
+			FString GetBackupExceptionMessage(const FString& ExceptionType, const FFrame& StackFrame)
+			{
+				// Generate as useful a message as we can if the exception didn't provide one.
+				FString FallbackExceptionMessage;
+				if (StackFrame.PreviousTrackingFrame &&
+					StackFrame.PreviousTrackingFrame->Object &&
+					StackFrame.PreviousTrackingFrame->Node)
+				{
+					FallbackExceptionMessage = FString::Printf(
+						TEXT("%s in %s.%s"), *ExceptionType,
+						*StackFrame.PreviousTrackingFrame->Object->GetName(),
+						*StackFrame.PreviousTrackingFrame->Node->GetName());
+				}
+				else
+				{
+					FallbackExceptionMessage = ExceptionType;
+				}
+				return FallbackExceptionMessage;
+			}
+
+			TOptional<FString> GetBlueprintExceptionMessage(const FBlueprintExceptionInfo& Info, const FFrame& StackFrame)
 			{
 				TOptional<FString> ExceptionMessage;
 				switch (Info.GetType())
@@ -82,32 +102,32 @@ namespace UE::ToolsetRegistry
 						break;
 					case EBlueprintExceptionType::AccessViolation:
 						ExceptionMessage = Info.GetDescription().IsEmpty() ?
-							TEXT("AccessViolation") :
+							GetBackupExceptionMessage(TEXT("AccessViolation"), StackFrame) :
 							Info.GetDescription().ToString();
 						break;
 					case EBlueprintExceptionType::InfiniteLoop:
 						ExceptionMessage = Info.GetDescription().IsEmpty() ?
-							TEXT("InfiniteLoop") :
+							GetBackupExceptionMessage(TEXT("InfiniteLoop"), StackFrame) :
 							Info.GetDescription().ToString();
 						break;
 					case EBlueprintExceptionType::NonFatalError:
 						ExceptionMessage = Info.GetDescription().IsEmpty() ?
-							TEXT("NonFatalError") :
+							GetBackupExceptionMessage(TEXT("NonFatalError"), StackFrame) :
 							Info.GetDescription().ToString();
 						break;
 					case EBlueprintExceptionType::FatalError:
 						ExceptionMessage = Info.GetDescription().IsEmpty() ?
-							TEXT("FatalError") :
+							GetBackupExceptionMessage(TEXT("FatalError"), StackFrame) :
 							Info.GetDescription().ToString();
 						break;
 					case EBlueprintExceptionType::AbortExecution:
 						ExceptionMessage = Info.GetDescription().IsEmpty() ?
-							TEXT("AbortExecution") :
+							GetBackupExceptionMessage(TEXT("AbortExecution"), StackFrame) :
 							Info.GetDescription().ToString();
 						break;
 					case EBlueprintExceptionType::UserRaisedError:
 						ExceptionMessage = Info.GetDescription().IsEmpty() ?
-							TEXT("UserRaisedError") :
+							GetBackupExceptionMessage(TEXT("UserRaisedError"), StackFrame) :
 							Info.GetDescription().ToString();
 						break;
 					default:
@@ -169,7 +189,7 @@ namespace UE::ToolsetRegistry
 		for (UFunction* Function : ValidTools)
 		{
 			check(Function);
-			
+
 			TSharedPtr<FJsonObject> ToolEntry =
 				UE::ToolsetRegistry::Internal::ToolsetJson::StructToJsonSchema(Function);
 
