Adding a toolset for setting and getting editor application state. First tools are for querying console commands and running console comands.

#Synced-CL 51402579

diff --git a/Engine/Plugins/Experimental/ToolsetRegistry/Source/ToolsetRegistry/Private/Tests/EditorAppToolsetTest.cpp b/Engine/Plugins/Experimental/ToolsetRegistry/Source/ToolsetRegistry/Private/Tests/EditorAppToolsetTest.cpp
new file mode 100644
--- /dev/null
+++ b/Engine/Plugins/Experimental/ToolsetRegistry/Source/ToolsetRegistry/Private/Tests/EditorAppToolsetTest.cpp
@@ -0,0 +1,128 @@
+// Copyright Epic Games, Inc. All Rights Reserved.
+
+#include "CoreMinimal.h"
+#include "Misc/AutomationTest.h"
+
+#include "ToolsetRegistry/EditorAppToolset.h"
+#include "ToolsetRegistry/JsonConversion.h"
+
+
+#if WITH_DEV_AUTOMATION_TESTS
+
+namespace
+{
+	template<typename CVarTypeT>
+	struct FCVarTestFixture
+	{
+		FCVarTestFixture(const CVarTypeT& InitialValue) :
+			CVarNamespace(TEXT("toolset.")),
+			CVarName(FString::Printf(TEXT("%stest.TestString"), *CVarNamespace)),
+			CVar(*CVarName, InitialValue, TEXT("A test cvar."))
+		{
+		}
+
+		TSharedPtr<FJsonObject> FindCVarJsonData() const
+		{
+			TSharedPtr<FJsonObject> CVarJsonData =
+				UE::ToolsetRegistry::Internal::JsonStringToJsonObject(
+					UEditorAppToolset::SearchCVars(CVarNamespace));
+			if (!CVarJsonData)
+			{
+				return nullptr;
+			}
+			const TSharedPtr<FJsonObject>* CommandJson = nullptr;
+			CVarJsonData->TryGetObjectField(CVarName, CommandJson);
+			return CommandJson ? *CommandJson : nullptr;
+		}
+
+		FString CVarNamespace;
+		FString CVarName;
+		TAutoConsoleVariable<CVarTypeT> CVar;
+	};
+}
+
+BEGIN_DEFINE_SPEC(FEditorAppToolsetSpec, "AI.ToolsetRegistry.EditorAppToolsetSpec",
+	EAutomationTestFlags::EditorContext | EAutomationTestFlags::EngineFilter)
+END_DEFINE_SPEC(FEditorAppToolsetSpec)
+
+void FEditorAppToolsetSpec::Define()
+{
+	It(TEXT("Can find a bool cvar"), [this]()
+	{
+		FCVarTestFixture<bool> Fixture(true);
+		TSharedPtr<FJsonObject> CommandJson = Fixture.FindCVarJsonData();
+		TestTrue("CVar", CommandJson.IsValid());
+		if (!CommandJson)
+		{
+			return;
+		}
+		TestEqual("CVar", CommandJson->GetStringField(
+			FString(TEXT("help"))), Fixture.CVar->GetHelp());
+		bool Value = false;
+		TestTrue("CVar", CommandJson->TryGetBoolField(
+			FString(TEXT("value")), Value));
+		TestEqual("CVar", Value, Fixture.CVar->AsVariable()->GetBool());
+	});
+
+	It(TEXT("Can find an int cvar"), [this]()
+	{
+		FCVarTestFixture<int> Fixture(5);
+		TSharedPtr<FJsonObject> CommandJson = Fixture.FindCVarJsonData();
+		TestTrue("CVar", CommandJson.IsValid());
+		if (!CommandJson)
+		{
+			return;
+		}
+		TestEqual("CVar", CommandJson->GetStringField(
+			FString(TEXT("help"))), Fixture.CVar->GetHelp());
+		int32 Value = 0;
+		TestTrue("CVar", CommandJson->TryGetNumberField(FString(TEXT("value")), Value));
+		TestEqual("CVar", Value, Fixture.CVar->AsVariable()->GetInt());
+	});
+
+	It(TEXT("Can find an float cvar"), [this]()
+	{
+		FCVarTestFixture<float> Fixture(5.0);
+		TSharedPtr<FJsonObject> CommandJson = Fixture.FindCVarJsonData();
+		TestTrue("CVar", CommandJson.IsValid());
+		if (!CommandJson)
+		{
+			return;
+		}
+		TestEqual("CVar", CommandJson->GetStringField(
+			FString(TEXT("help"))), Fixture.CVar->GetHelp());
+		float Value = 0;
+		TestTrue("CVar", CommandJson->TryGetNumberField(
+			FString(TEXT("value")), Value));
+		TestEqual("CVar", Value, Fixture.CVar->AsVariable()->GetFloat());
+	});
+
+	It(TEXT("Can find an string cvar"), [this]()
+	{
+		FCVarTestFixture<FString> Fixture(TEXT("testy"));
+		TSharedPtr<FJsonObject> CommandJson = Fixture.FindCVarJsonData();
+		TestTrue("CVar", CommandJson.IsValid());
+		if (!CommandJson)
+		{
+			return;
+		}
+		TestEqual("CVar", CommandJson->GetStringField(
+			FString(TEXT("help"))), Fixture.CVar->GetHelp());
+		FString Value;
+		TestTrue("CVar", CommandJson->TryGetStringField(
+			FString(TEXT("value")), Value));
+		TestEqual("CVar", Value, Fixture.CVar->AsVariable()->GetString());
+	});
+
+	It(TEXT("Skips disabled cvar"), [this]()
+	{
+		FString CVarName(TEXT("toolset.test.TestDisabled"));
+		TAutoConsoleVariable<bool> CVar(*CVarName, true, TEXT("A disabled cvar."), ECVF_Unregistered);
+		FString CVarJsonString = UEditorAppToolset::SearchCVars(CVarName);
+		TSharedPtr<FJsonObject> CVarJsonData = UE::ToolsetRegistry::Internal::JsonStringToJsonObject(CVarJsonString);
+		TestTrue("CVar", CVarJsonData.IsValid());
+		TestEqual("CVar", CVarJsonData->Values.Num(), 0);
+	});
+}
+
+#endif
diff --git a/Engine/Plugins/Experimental/ToolsetRegistry/Source/ToolsetRegistry/Private/ToolsetRegistry/EditorAppToolset.cpp b/Engine/Plugins/Experimental/ToolsetRegistry/Source/ToolsetRegistry/Private/ToolsetRegistry/EditorAppToolset.cpp
new file mode 100644
--- /dev/null
+++ b/Engine/Plugins/Experimental/ToolsetRegistry/Source/ToolsetRegistry/Private/ToolsetRegistry/EditorAppToolset.cpp
@@ -0,0 +1,64 @@
+// Copyright Epic Games, Inc. All Rights Reserved.
+
+#include "EditorAppToolset.h"
+
+#include "HAL/ConsoleManager.h"
+#include "Kismet/KismetSystemLibrary.h"
+
+#include "ToolsetRegistry/JsonConversion.h"
+
+#include UE_INLINE_GENERATED_CPP_BY_NAME(EditorAppToolset)
+
+FString UEditorAppToolset::SearchCVars(const FString& Name)
+{
+	TSharedPtr<FJsonObject> Results = MakeShared<FJsonObject>();
+	const auto OnConsoleVariable = [&Results](const TCHAR* Name, IConsoleObject* CVar)
+	{
+		if (TSharedPtr<FJsonObject> CVarJson = CVarToJson(CVar->AsVariable()))
+		{
+			Results->SetObjectField(Name, CVarJson);
+		}
+	};
+	IConsoleManager& ConsoleManager = IConsoleManager::Get();
+	ConsoleManager.ForEachConsoleObjectThatContains(FConsoleObjectVisitor::CreateLambda(OnConsoleVariable), *Name);
+	return UE::ToolsetRegistry::Internal::JsonToString(Results.ToSharedRef());
+}
+
+TSharedPtr<FJsonObject> UEditorAppToolset::CVarToJson(IConsoleObject* CVar)
+{
+	if (!CVar || !CVar->IsEnabled())
+	{
+		return nullptr;
+	}
+
+	TSharedPtr<FJsonObject> CVarData = MakeShared<FJsonObject>();
+	CVarData->SetStringField(FString(TEXT("help")), CVar->GetHelp());
+	if (IConsoleVariable* Variable = CVar->AsVariable())
+	{
+		if (Variable->IsVariableBool())
+		{
+			CVarData->SetBoolField(
+				FString(TEXT("value")), Variable->GetBool());
+		}
+		else if (Variable->IsVariableInt())
+		{
+			CVarData->SetNumberField(
+				FString(TEXT("value")), Variable->GetInt());
+		}
+		else if (Variable->IsVariableFloat())
+		{
+			CVarData->SetNumberField(
+				FString(TEXT("value")), Variable->GetFloat());
+		}
+		else if (Variable->IsVariableString())
+		{
+			CVarData->SetStringField(
+				FString(TEXT("value")), Variable->GetString());
+		}
+		else
+		{
+			check(false);
+		}
+	}
+	return CVarData;
+}
diff --git a/Engine/Plugins/Experimental/ToolsetRegistry/Source/ToolsetRegistry/Private/ToolsetRegistry/EditorAppToolset.h b/Engine/Plugins/Experimental/ToolsetRegistry/Source/ToolsetRegistry/Private/ToolsetRegistry/EditorAppToolset.h
new file mode 100644
--- /dev/null
+++ b/Engine/Plugins/Experimental/ToolsetRegistry/Source/ToolsetRegistry/Private/ToolsetRegistry/EditorAppToolset.h
@@ -0,0 +1,32 @@
+// Copyright Epic Games, Inc. All Rights Reserved.
+
+#pragma once
+
+#include "CoreMinimal.h"
+
+#include "ToolsetRegistry/UToolsetRegistry.h"
+
+#include "EditorAppToolset.generated.h"
+
+class FJsonObject;
+class IConsoleObject;
+
+/// Contains tools that query or modify the state of the Unreal Editor application.
+UCLASS(BlueprintType)
+class TOOLSETREGISTRY_API UEditorAppToolset : public UToolsetDefinition
+{
+	GENERATED_BODY()
+public:
+	/*
+	 * Finds all console variables that contain a given name.
+	 * @param Name The partial or full name to search for.
+	 * @return A JSON dict with info about matching cvars.
+	 */
+	UFUNCTION(BlueprintCallable, Category = "EditorAppToolset")
+	static FString SearchCVars(const FString& Name);
+
+private:
+	static TSharedPtr<FJsonObject> CVarToJson(IConsoleObject* CVar);
+
+	friend class FEditorAppToolsetSpec;
+};
diff --git a/Engine/Plugins/Experimental/ToolsetRegistry/Source/ToolsetRegistry/Private/ToolsetRegistry/ToolsetRegistrySubsystem.cpp b/Engine/Plugins/Experimental/ToolsetRegistry/Source/ToolsetRegistry/Private/ToolsetRegistry/ToolsetRegistrySubsystem.cpp
--- a/Engine/Plugins/Experimental/ToolsetRegistry/Source/ToolsetRegistry/Private/ToolsetRegistry/ToolsetRegistrySubsystem.cpp
+++ b/Engine/Plugins/Experimental/ToolsetRegistry/Source/ToolsetRegistry/Private/ToolsetRegistry/ToolsetRegistrySubsystem.cpp
@@ -6,6 +6,7 @@
 #include "Subsystems/SubsystemCollection.h"
 
 #include "ToolsetRegistry/AgentSkill.h"
+#include "ToolsetRegistry/EditorAppToolset.h"
 #include "ToolsetRegistry/Module.h"
 #include "ToolsetRegistry/UToolsetRegistry.h"
 
@@ -41,11 +42,13 @@ TValueOrError<TObjectPtr<UToolsetRegistrySubsystem>, FString> UToolsetRegistrySu
 void UToolsetRegistrySubsystem::Initialize(FSubsystemCollectionBase& Collection)
 {
 	Super::Initialize(Collection);
+	UToolsetRegistry::RegisterToolsetClass(UEditorAppToolset::StaticClass());
 	UToolsetRegistry::RegisterToolsetClass(UAgentSkillToolset::StaticClass());
 }
 
 void UToolsetRegistrySubsystem::Deinitialize()
 {
 	UToolsetRegistry::UnregisterToolsetClass(UAgentSkillToolset::StaticClass());
+	UToolsetRegistry::UnregisterToolsetClass(UEditorAppToolset::StaticClass());
 	Super::Deinitialize();
 }
\ No newline at end of file
