[AI Assistant] Only export required symbols from public classes and structs.

The toolset registry now only exports public symbols that are required by
users of the library.

#nojira

#Synced-CL 51484694

diff --git a/Engine/Plugins/Experimental/ToolsetRegistry/Source/ToolsetRegistry/Private/ToolsetRegistry/ToolsetLibrary.h b/Engine/Plugins/Experimental/ToolsetRegistry/Source/ToolsetRegistry/Private/ToolsetRegistry/ToolsetLibrary.h
--- a/Engine/Plugins/Experimental/ToolsetRegistry/Source/ToolsetRegistry/Private/ToolsetRegistry/ToolsetLibrary.h
+++ b/Engine/Plugins/Experimental/ToolsetRegistry/Source/ToolsetRegistry/Private/ToolsetRegistry/ToolsetLibrary.h
@@ -10,8 +10,8 @@
 
 /// Provides functions that are critical to Python toolset operation but which are not available
 /// in and cannot be added to any other BFPL.
-UCLASS(BlueprintType)
-class TOOLSETREGISTRY_API UToolsetLibrary : public UBlueprintFunctionLibrary
+UCLASS(BlueprintType, MinimalAPI)
+class UToolsetLibrary : public UBlueprintFunctionLibrary
 {
 	GENERATED_BODY()
 public:
diff --git a/Engine/Plugins/Experimental/ToolsetRegistry/Source/ToolsetRegistry/Public/ToolsetRegistry/AgentSkill.h b/Engine/Plugins/Experimental/ToolsetRegistry/Source/ToolsetRegistry/Public/ToolsetRegistry/AgentSkill.h
--- a/Engine/Plugins/Experimental/ToolsetRegistry/Source/ToolsetRegistry/Public/ToolsetRegistry/AgentSkill.h
+++ b/Engine/Plugins/Experimental/ToolsetRegistry/Source/ToolsetRegistry/Public/ToolsetRegistry/AgentSkill.h
@@ -9,6 +9,8 @@
 
 #include "AgentSkill.generated.h"
 
+#define UE_API TOOLSETREGISTRY_API
+
 class UBlueprintFunctionLibrary;
 
 /// Detailed information about a specific AgentSkill.
@@ -16,6 +18,7 @@ USTRUCT(BlueprintType)
 struct FAgentSkillDetails
 {
 	GENERATED_USTRUCT_BODY()
+
 public:
 	/// The full names of the toolsets the skill depends on.
 	UPROPERTY(BlueprintReadWrite, Category = "AgentSkill")
@@ -26,17 +29,18 @@ public:
 	FString Instructions;
 };
 
-UCLASS(Blueprintable)
-class TOOLSETREGISTRY_API UAgentSkill : public UObject
+UCLASS(Blueprintable, MinimalAPI)
+class UAgentSkill : public UObject
 {
 	GENERATED_BODY()
+
 public:
 	/// A brief description of what the skill does or how it should be used.
 	UPROPERTY(EditDefaultsOnly, Category="AgentSkill", meta=(MultiLine=true))
 	FString Description;
 
 	/// Returns detailed information about how to use this skill.
-	FAgentSkillDetails GetDetails() const;
+	UE_API FAgentSkillDetails GetDetails() const;
 
 protected:
 	/// Generates the final prompt text for the Agent Skill.
@@ -62,9 +66,10 @@ protected:
 
 
 UCLASS(BlueprintType, Hidden)
-class TOOLSETREGISTRY_API UAgentSkillToolset : public UToolsetDefinition
+class UAgentSkillToolset : public UToolsetDefinition
 {
 	GENERATED_BODY()
+
 public:
 	/**
 	 * Gets a summary of all AgentSkills in the project.
@@ -106,3 +111,5 @@ public:
 	static bool UpdateSkill(
 		const FString& SkillPath, const FString& Description, const FAgentSkillDetails& Details);
 };
+
+#undef UE_API
diff --git a/Engine/Plugins/Experimental/ToolsetRegistry/Source/ToolsetRegistry/Public/ToolsetRegistry/ToolCallAsyncResultString.h b/Engine/Plugins/Experimental/ToolsetRegistry/Source/ToolsetRegistry/Public/ToolsetRegistry/ToolCallAsyncResultString.h
--- a/Engine/Plugins/Experimental/ToolsetRegistry/Source/ToolsetRegistry/Public/ToolsetRegistry/ToolCallAsyncResultString.h
+++ b/Engine/Plugins/Experimental/ToolsetRegistry/Source/ToolsetRegistry/Public/ToolsetRegistry/ToolCallAsyncResultString.h
@@ -17,7 +17,7 @@ class UToolCallAsyncResultString : public UToolCallAsyncResult
 public:
 	// Set the result.
 	UFUNCTION(BlueprintCallable, Category = "ToolsetRegistry")
-	UE_API bool SetValue(const FString& InValue)
+	bool SetValue(const FString& InValue)
 	{
 		return MaybeBroadcastSuccessfulCompletion(FString(InValue), Value);
 	}
diff --git a/Engine/Plugins/Experimental/ToolsetRegistry/Source/ToolsetRegistry/Public/ToolsetRegistry/ToolCallAsyncResultVoid.h b/Engine/Plugins/Experimental/ToolsetRegistry/Source/ToolsetRegistry/Public/ToolsetRegistry/ToolCallAsyncResultVoid.h
--- a/Engine/Plugins/Experimental/ToolsetRegistry/Source/ToolsetRegistry/Public/ToolsetRegistry/ToolCallAsyncResultVoid.h
+++ b/Engine/Plugins/Experimental/ToolsetRegistry/Source/ToolsetRegistry/Public/ToolsetRegistry/ToolCallAsyncResultVoid.h
@@ -22,7 +22,7 @@ class UToolCallAsyncResultVoid : public UToolCallAsyncResult
 public:
 	// Complete this result with no value and notify listeners of OnCompleted.
 	UFUNCTION(BlueprintCallable, Category = "ToolsetRegistry")
-	UE_API bool SetCompleted()
+	bool SetCompleted()
 	{
 		return MaybeBroadcastSuccessfulCompletion();
 	}
diff --git a/Engine/Plugins/Experimental/ToolsetRegistry/Source/ToolsetRegistry/Public/ToolsetRegistry/ToolsetImage.h b/Engine/Plugins/Experimental/ToolsetRegistry/Source/ToolsetRegistry/Public/ToolsetRegistry/ToolsetImage.h
--- a/Engine/Plugins/Experimental/ToolsetRegistry/Source/ToolsetRegistry/Public/ToolsetRegistry/ToolsetImage.h
+++ b/Engine/Plugins/Experimental/ToolsetRegistry/Source/ToolsetRegistry/Public/ToolsetRegistry/ToolsetImage.h
@@ -6,9 +6,11 @@
 
 #include "ToolsetImage.generated.h"
 
+#define UE_API TOOLSETREGISTRY_API
+
 /// The standard image format that toolsets should return.
-USTRUCT(BlueprintType)
-struct TOOLSETREGISTRY_API FToolsetImage
+USTRUCT(BlueprintType, MinimalAPI)
+struct FToolsetImage
 {
 	GENERATED_BODY()
 public:
@@ -21,8 +23,10 @@ public:
 	FString Data;
 	
 	/// Populates the ToolsetImage based on a raw color image.
-	bool SetFromBitmap(const TArray<FColor>& Bitmap, FIntPoint Dimensions);
+	UE_API bool SetFromBitmap(const TArray<FColor>& Bitmap, FIntPoint Dimensions);
 
 	/// Populates the ToolsetImage from a file saved on disk.
-	bool SetFromFile(const FString& Path);
+	UE_API bool SetFromFile(const FString& Path);
 };
+
+#undef UE_API
\ No newline at end of file
diff --git a/Engine/Plugins/Experimental/ToolsetRegistry/Source/ToolsetRegistry/Public/ToolsetRegistry/ToolsetRegistrySubsystem.h b/Engine/Plugins/Experimental/ToolsetRegistry/Source/ToolsetRegistry/Public/ToolsetRegistry/ToolsetRegistrySubsystem.h
--- a/Engine/Plugins/Experimental/ToolsetRegistry/Source/ToolsetRegistry/Public/ToolsetRegistry/ToolsetRegistrySubsystem.h
+++ b/Engine/Plugins/Experimental/ToolsetRegistry/Source/ToolsetRegistry/Public/ToolsetRegistry/ToolsetRegistrySubsystem.h
@@ -12,18 +12,20 @@
 
 #include "ToolsetRegistrySubsystem.generated.h"
 
+#define UE_API TOOLSETREGISTRY_API
+
 class FSubsystemCollectionBase;
 
-UCLASS(BlueprintType)
-class TOOLSETREGISTRY_API UToolsetRegistrySubsystem : public UEditorSubsystem
+UCLASS(BlueprintType, MinimalAPI)
+class UToolsetRegistrySubsystem : public UEditorSubsystem
 {
 	GENERATED_BODY()
 
 public:
-
-	// Get the toolset registry subsystem.  Optionally log a warning if passed a non-empty string (the warning
-	// message will start with the passed-in string).
-	static TValueOrError<TObjectPtr<UToolsetRegistrySubsystem>, FString> Get(const FString& WarningMessage = FString());
+	// Get the toolset registry subsystem.  Optionally log a warning if passed a non-empty string
+	// (the warning message will start with the passed-in string).
+	UE_API static TValueOrError<TObjectPtr<UToolsetRegistrySubsystem>, FString> Get(
+		const FString& WarningMessage = FString());
 	
 	virtual void Initialize(FSubsystemCollectionBase& Collection) override;
 	virtual void Deinitialize() override;
@@ -32,3 +34,4 @@ public:
 	UE::ToolsetRegistry::FToolsetRegistry ToolsetRegistry;
 };
 
+#undef UE_API
diff --git a/Engine/Plugins/Experimental/ToolsetRegistry/Source/ToolsetRegistry/Public/ToolsetRegistry/UToolsetRegistry.h b/Engine/Plugins/Experimental/ToolsetRegistry/Source/ToolsetRegistry/Public/ToolsetRegistry/UToolsetRegistry.h
--- a/Engine/Plugins/Experimental/ToolsetRegistry/Source/ToolsetRegistry/Public/ToolsetRegistry/UToolsetRegistry.h
+++ b/Engine/Plugins/Experimental/ToolsetRegistry/Source/ToolsetRegistry/Public/ToolsetRegistry/UToolsetRegistry.h
@@ -8,47 +8,49 @@
 
 #include "UToolsetRegistry.generated.h"
 
+#define UE_API TOOLSETREGISTRY_API
+
 // This is the common base class for Toolsets defined as UObjects.
 // UFunctions that define tools in this class should be static functions.
-UCLASS(BlueprintType, Abstract)
-class TOOLSETREGISTRY_API UToolsetDefinition : public UObject
+UCLASS(BlueprintType, Abstract, MinimalAPI)
+class UToolsetDefinition : public UObject
 {
 	GENERATED_BODY()
-public:
 };
 
 // Blueprint-accessible wrapper to allow queries against the AI Toolset Registry.
-UCLASS()
-class TOOLSETREGISTRY_API UToolsetRegistry : public UBlueprintFunctionLibrary
+UCLASS(MinimalAPI)
+class UToolsetRegistry : public UBlueprintFunctionLibrary
 {
 	GENERATED_BODY()
 
 public:
-
 	// Register a Blueprint function library class as a toolset.
 	UFUNCTION(BlueprintCallable, Category="ToolsetRegistry")
-	static void RegisterToolsetClass(TSubclassOf<UToolsetDefinition> InToolsetClass);
+	static UE_API void RegisterToolsetClass(TSubclassOf<UToolsetDefinition> InToolsetClass);
 
 	// Unregister a Blueprint function library class as a toolset.
 	UFUNCTION(BlueprintCallable, Category="ToolsetRegistry")
-	static void UnregisterToolsetClass(TSubclassOf<UToolsetDefinition> InToolsetClass);
+	static UE_API void UnregisterToolsetClass(TSubclassOf<UToolsetDefinition> InToolsetClass);
 
  	// Return whether a toolset class is already registered.
 	UFUNCTION(BlueprintPure, Category="ToolsetRegistry")
-	static bool IsToolsetClassRegistered(TSubclassOf<UToolsetDefinition> InToolsetClass);
+	static UE_API bool IsToolsetClassRegistered(TSubclassOf<UToolsetDefinition> InToolsetClass);
 
  	// Return whether a toolset is already registered by name.
 	UFUNCTION(BlueprintPure, Category="ToolsetRegistry")
-	static bool IsToolsetRegistered(const FString& InToolsetName);
+	static UE_API bool IsToolsetRegistered(const FString& InToolsetName);
 
 	// Execute a registered tool given a toolset and tool name as if it was being executed by an AI assistant.
 	// Meant for use while testing.
 	UFUNCTION(BlueprintCallable, Category="ToolsetRegistry")
-	static void ExecuteTool(
+	static UE_API void ExecuteTool(
 		const FString& ToolsetName, const FString& ToolName, const FString& JsonInput, FString& OutResult, FString& OutError);
 
 	// Get toolset JSON schema for a toolset class.
 	// Meant for use while testing.
 	UFUNCTION(BlueprintPure, Category="ToolsetRegistry")
-	static FString GetToolsetJsonSchema(TSubclassOf<UToolsetDefinition> InToolsetClass);
+	static UE_API FString GetToolsetJsonSchema(TSubclassOf<UToolsetDefinition> InToolsetClass);
 };
+
+#undef UE_API
