[AI Assistant] Fix more unnceessarily exported symbols.

#Synced-CL 51517130

diff --git a/Engine/Plugins/Experimental/ToolsetRegistry/Source/ToolsetRegistry/Private/ToolsetRegistry/EditorAppToolset.h b/Engine/Plugins/Experimental/ToolsetRegistry/Source/ToolsetRegistry/Private/ToolsetRegistry/EditorAppToolset.h
--- a/Engine/Plugins/Experimental/ToolsetRegistry/Source/ToolsetRegistry/Private/ToolsetRegistry/EditorAppToolset.h
+++ b/Engine/Plugins/Experimental/ToolsetRegistry/Source/ToolsetRegistry/Private/ToolsetRegistry/EditorAppToolset.h
@@ -12,8 +12,8 @@ class FJsonObject;
 class IConsoleObject;
 
 /// Contains tools that query or modify the state of the Unreal Editor application.
-UCLASS(BlueprintType)
-class TOOLSETREGISTRY_API UEditorAppToolset : public UToolsetDefinition
+UCLASS(BlueprintType, MinimalAPI)
+class UEditorAppToolset : public UToolsetDefinition
 {
 	GENERATED_BODY()
 public:
diff --git a/Engine/Plugins/Experimental/ToolsetRegistry/Source/ToolsetRegistry/Private/ToolsetRegistry/ToolsetLibrary.h b/Engine/Plugins/Experimental/ToolsetRegistry/Source/ToolsetRegistry/Private/ToolsetRegistry/ToolsetLibrary.h
--- a/Engine/Plugins/Experimental/ToolsetRegistry/Source/ToolsetRegistry/Private/ToolsetRegistry/ToolsetLibrary.h
+++ b/Engine/Plugins/Experimental/ToolsetRegistry/Source/ToolsetRegistry/Private/ToolsetRegistry/ToolsetLibrary.h
@@ -53,8 +53,8 @@ public:
 // Provides workarounds for the lack of async tool call support. Once we have that,
 // we should remove these functions and implement a tool that directly takes a screenshot and
 // returns a ToolsetImage.
-UCLASS(BlueprintType)
-class TOOLSETREGISTRY_API UToolsetLibraryTemp : public UBlueprintFunctionLibrary
+UCLASS(BlueprintType, MinimalAPI)
+class UToolsetLibraryTemp : public UBlueprintFunctionLibrary
 {
 	GENERATED_BODY()
 public:
