Add fallback handling for cvars that are variables without a specific type.

#Synced-CL 51567995

diff --git a/Engine/Plugins/Experimental/ToolsetRegistry/Source/ToolsetRegistry/Private/Tests/EditorAppToolsetTest.cpp b/Engine/Plugins/Experimental/ToolsetRegistry/Source/ToolsetRegistry/Private/Tests/EditorAppToolsetTest.cpp
--- a/Engine/Plugins/Experimental/ToolsetRegistry/Source/ToolsetRegistry/Private/Tests/EditorAppToolsetTest.cpp
+++ b/Engine/Plugins/Experimental/ToolsetRegistry/Source/ToolsetRegistry/Private/Tests/EditorAppToolsetTest.cpp
@@ -114,6 +114,20 @@ void FEditorAppToolsetSpec::Define()
 		TestEqual("CVar", Value, Fixture.CVar->AsVariable()->GetString());
 	});
 
+	It(TEXT("Can get default values from cvars"), [this]()
+	{
+		IConsoleManager& ConsoleManager = IConsoleManager::Get();
+		FString CVarName(TEXT("toolset.test.TestBit"));
+		static uint8 Force0Mask[1] = { 0 };
+		static uint8 Force1Mask[1] = { 0 };
+		const uint32 BitNumber = 0;
+		IConsoleVariable* CVar = ConsoleManager.RegisterConsoleVariableBitRef(
+			*CVarName, TEXT("MyFlag"), BitNumber, Force0Mask, Force1Mask, TEXT("Test"));
+		TestTrue("CVar", !UEditorAppToolset::SearchCVars(CVarName).IsEmpty());
+		ConsoleManager.UnregisterConsoleObject(CVar);
+	});
+
+
 	It(TEXT("Skips disabled cvar"), [this]()
 	{
 		FString CVarName(TEXT("toolset.test.TestDisabled"));
diff --git a/Engine/Plugins/Experimental/ToolsetRegistry/Source/ToolsetRegistry/Private/ToolsetRegistry/EditorAppToolset.cpp b/Engine/Plugins/Experimental/ToolsetRegistry/Source/ToolsetRegistry/Private/ToolsetRegistry/EditorAppToolset.cpp
--- a/Engine/Plugins/Experimental/ToolsetRegistry/Source/ToolsetRegistry/Private/ToolsetRegistry/EditorAppToolset.cpp
+++ b/Engine/Plugins/Experimental/ToolsetRegistry/Source/ToolsetRegistry/Private/ToolsetRegistry/EditorAppToolset.cpp
@@ -57,7 +57,7 @@ TSharedPtr<FJsonObject> UEditorAppToolset::CVarToJson(IConsoleObject* CVar)
 		}
 		else
 		{
-			check(false);
+			CVarData->SetStringField(TEXT("value"), Variable->GetDefaultValue());
 		}
 	}
 	return CVarData;
