[AI Assistant] Don't log an error on unregistration if the registry isn't loaded.

If the toolset registry is not loaded, no toolsets are registered and hence
unregistration of a toolset should do nothing. This changes
`UToolsetRegistry::UnregisterToolsetClass` to only warn if the toolset
registry subsystem is not present.

#Synced-CL 51567995

diff --git a/Engine/Plugins/Experimental/ToolsetRegistry/Source/ToolsetRegistry/Private/ToolsetRegistry/UToolsetRegistry.cpp b/Engine/Plugins/Experimental/ToolsetRegistry/Source/ToolsetRegistry/Private/ToolsetRegistry/UToolsetRegistry.cpp
--- a/Engine/Plugins/Experimental/ToolsetRegistry/Source/ToolsetRegistry/Private/ToolsetRegistry/UToolsetRegistry.cpp
+++ b/Engine/Plugins/Experimental/ToolsetRegistry/Source/ToolsetRegistry/Private/ToolsetRegistry/UToolsetRegistry.cpp
@@ -63,7 +63,7 @@ void UToolsetRegistry::UnregisterToolsetClass(TSubclassOf<UToolsetDefinition> In
 	}
 
 	if (auto ToolsetRegistrySubsystem =
-			UToolsetRegistrySubsystem::Get(TEXT("IsToolsetClassRegistered"));
+			UToolsetRegistrySubsystem::Get(TEXT("UnregisterToolsetClass"));
 		ToolsetRegistrySubsystem.HasValue())
 	{
 		FString ToolsetName =
@@ -76,10 +76,6 @@ void UToolsetRegistry::UnregisterToolsetClass(TSubclassOf<UToolsetDefinition> In
 			ToolsetRegistrySubsystem.GetValue()->ToolsetRegistry.UnregisterToolset(Handler);
 		}
 	}
-	else
-	{
-		UE_LOG(LogToolsetRegistry, Error, TEXT("AIToolsetRegistrySubsystem unavailable"));
-	}
 }
 
 bool UToolsetRegistry::IsToolsetClassRegistered(TSubclassOf<UToolsetDefinition> InToolsetClass)
