Added toolset for working with data tables.

#Synced-CL 51567995

diff --git a/Engine/Plugins/Experimental/ToolsetRegistry/Source/ToolsetRegistry/Private/Tests/ToolsetLibraryTest.cpp b/Engine/Plugins/Experimental/ToolsetRegistry/Source/ToolsetRegistry/Private/Tests/ToolsetLibraryTest.cpp
--- a/Engine/Plugins/Experimental/ToolsetRegistry/Source/ToolsetRegistry/Private/Tests/ToolsetLibraryTest.cpp
+++ b/Engine/Plugins/Experimental/ToolsetRegistry/Source/ToolsetRegistry/Private/Tests/ToolsetLibraryTest.cpp
@@ -4,6 +4,7 @@
 #include "JsonObjectConverter.h"
 #include "Misc/AutomationTest.h"
 
+#include "Animation/AnimNode_Root.h"
 #include "Tests/BlueprintLibraryToolsetTest.h"
 #include "ToolsetJsonTest.h"
 #include "ToolsetRegistry/JsonConversion.h"
@@ -22,7 +23,7 @@ void FToolsetLibrarySpec::Define()
 
 	It("Generates schema from properties.", [this]()
 	{
-		FString PropertiesJson = UToolsetLibrary::ListClassProperties(
+		FString PropertiesJson = UToolsetLibrary::ListStructProperties(
 			UToolsetJsonTestObject::StaticClass());
 		TSharedPtr<FJsonObject> JsonObject = JsonStringToJsonObject(PropertiesJson);
 		TestTrue("Schema", JsonObject.IsValid());
@@ -77,6 +78,13 @@ void FToolsetLibrarySpec::Define()
 			UToolsetDefinition::StaticClass());
 		TestTrue("Subclass", Subclasses.Contains(UFakeToolset::StaticClass()));
 	});
+
+	It("Gets substructs of a struct.", [this]()
+	{
+		TArray<UScriptStruct*> Substructs = UToolsetLibrary::GetDerivedStructs(
+			FAnimNode_Base::StaticStruct());
+		TestTrue("Subclass", Substructs.Contains(FAnimNode_Root::StaticStruct()));
+	});
 }
 
 #endif
diff --git a/Engine/Plugins/Experimental/ToolsetRegistry/Source/ToolsetRegistry/Private/ToolsetRegistry/ToolsetLibrary.cpp b/Engine/Plugins/Experimental/ToolsetRegistry/Source/ToolsetRegistry/Private/ToolsetRegistry/ToolsetLibrary.cpp
--- a/Engine/Plugins/Experimental/ToolsetRegistry/Source/ToolsetRegistry/Private/ToolsetRegistry/ToolsetLibrary.cpp
+++ b/Engine/Plugins/Experimental/ToolsetRegistry/Source/ToolsetRegistry/Private/ToolsetRegistry/ToolsetLibrary.cpp
@@ -5,6 +5,7 @@
 #include "AssetRegistry/AssetRegistryHelpers.h"
 #include "AssetRegistry/AssetRegistryModule.h"
 #include "Dom/JsonObject.h"
+#include "UObject/UObjectIterator.h"
 #include "Misc/AssertionMacros.h"
 #include "UObject/PropertyAccessUtil.h"
 #include "UnrealClient.h"
@@ -15,15 +16,15 @@
 
 #include UE_INLINE_GENERATED_CPP_BY_NAME(ToolsetLibrary)
 
-FString UToolsetLibrary::ListClassProperties(UClass* Class)
+FString UToolsetLibrary::ListStructProperties(UStruct* Struct)
 {
-	if (!Class)
+	if (!Struct)
 	{
 		return FString();
 	}
 	FString JsonProperties;
 	TSharedPtr<FJsonObject> ClassSchema =
-		UE::ToolsetRegistry::Internal::ToolsetJson::StructToJsonSchema(Class, true);
+		UE::ToolsetRegistry::Internal::ToolsetJson::StructToJsonSchema(Struct, true);
 	const TSharedPtr<FJsonObject>* ClassProperties = nullptr;
 	if (ClassSchema->TryGetObjectField(TEXT("properties"), ClassProperties))
 	{
@@ -91,12 +92,10 @@ bool UToolsetLibrary::SetObjectProperties(UObject* Object, const FString& Proper
 
 TArray<UClass*> UToolsetLibrary::GetDerivedClasses(UClass* BaseClass)
 {
-	if (!BaseClass)
-	{
-		return TArray<UClass*>();
-	}
+	check(BaseClass);
 
-	FAssetRegistryModule& AssetRegistryModule = FModuleManager::LoadModuleChecked<FAssetRegistryModule>("AssetRegistry");
+	FAssetRegistryModule& AssetRegistryModule =
+		FModuleManager::LoadModuleChecked<FAssetRegistryModule>("AssetRegistry");
 	IAssetRegistry& AssetRegistry = AssetRegistryModule.Get();
 
 	TSet<FTopLevelAssetPath> DerivedNames;
@@ -124,6 +123,22 @@ TArray<UClass*> UToolsetLibrary::GetDerivedClasses(UClass* BaseClass)
 	return DerivedClasses;
 }
 
+TArray<UScriptStruct*> UToolsetLibrary::GetDerivedStructs(UScriptStruct* BaseStruct)
+{
+	check(BaseStruct);
+
+	TArray<UScriptStruct*> DerivedStructs;
+	for (TObjectIterator<UScriptStruct> It; It; ++It)
+	{
+		UScriptStruct* Struct = *It;
+		if (Struct != BaseStruct && Struct->IsChildOf(BaseStruct))
+		{
+			DerivedStructs.Add(Struct);
+		}
+	}
+	return DerivedStructs;
+}
+
 FString UToolsetLibraryTemp::ScreenshotPath(const FString& Filename)
 {
 	FString FullPath = Filename;
diff --git a/Engine/Plugins/Experimental/ToolsetRegistry/Source/ToolsetRegistry/Private/ToolsetRegistry/ToolsetLibrary.h b/Engine/Plugins/Experimental/ToolsetRegistry/Source/ToolsetRegistry/Private/ToolsetRegistry/ToolsetLibrary.h
--- a/Engine/Plugins/Experimental/ToolsetRegistry/Source/ToolsetRegistry/Private/ToolsetRegistry/ToolsetLibrary.h
+++ b/Engine/Plugins/Experimental/ToolsetRegistry/Source/ToolsetRegistry/Private/ToolsetRegistry/ToolsetLibrary.h
@@ -17,11 +17,11 @@ class UToolsetLibrary : public UBlueprintFunctionLibrary
 public:
 	/*
 	 * Returns the Blueprint visible properties in the class as JSON Schema.
-	 * @param Class The class to extract properties from.
+	 * @param Struct The struct to extract properties from.
 	 * @return A JSON Schema formatted string that describes the properties.
 	 */
 	UFUNCTION(BlueprintCallable, Category = "Toolset Library")
-	static FString ListClassProperties(UClass* Class);
+	static FString ListStructProperties(UStruct* Struct);
 
 	/*
 	 * Returns the values of the requested properties.
@@ -43,11 +43,21 @@ public:
 
 	/*
 	 * Returns the list of subclasses that derive from a class.
+	 * Uses the Asset Registry to find native and BP subclasses.
 	 * @param BaseClass The class to get subclasses from.
 	 * @return The list of native and BP classes that derive from the base class.
 	 */
 	UFUNCTION(BlueprintCallable, Category = "Toolset Library")
 	static TArray<UClass*> GetDerivedClasses(UClass* BaseClass);
+
+	/*
+	 * Returns the list of substructs that derive from a struct.
+	 * Iterates all loaded UScriptStruct objects to find matching substructs.
+	 * @param BaseStruct The struct to get substructs from.
+	 * @return The list of loaded structs that derive from the base struct.
+	 */
+	UFUNCTION(BlueprintCallable, Category = "Toolset Library")
+	static TArray<UScriptStruct*> GetDerivedStructs(UScriptStruct* BaseStruct);
 };
 
 // Provides workarounds for the lack of async tool call support. Once we have that,
