[AI Assistant]: Silence "Editor is not available" warning on toolset unregistration.

Similar to in CL 51638282, on shutdown if the toolset registry is not available then
all toolsets have already been unregistered so the unregistration operation does
nothing and should log nothing.

#Synced-CL 51655650

diff --git a/Engine/Plugins/Experimental/ToolsetRegistry/Source/ToolsetRegistry/Private/ToolsetRegistry/UToolsetRegistry.cpp b/Engine/Plugins/Experimental/ToolsetRegistry/Source/ToolsetRegistry/Private/ToolsetRegistry/UToolsetRegistry.cpp
--- a/Engine/Plugins/Experimental/ToolsetRegistry/Source/ToolsetRegistry/Private/ToolsetRegistry/UToolsetRegistry.cpp
+++ b/Engine/Plugins/Experimental/ToolsetRegistry/Source/ToolsetRegistry/Private/ToolsetRegistry/UToolsetRegistry.cpp
@@ -62,8 +62,7 @@ void UToolsetRegistry::UnregisterToolsetClass(TSubclassOf<UToolsetDefinition> In
 		return;
 	}
 
-	if (auto ToolsetRegistrySubsystem =
-			UToolsetRegistrySubsystem::Get(TEXT("UnregisterToolsetClass"));
+	if (auto ToolsetRegistrySubsystem = UToolsetRegistrySubsystem::Get();
 		ToolsetRegistrySubsystem.HasValue())
 	{
 		FString ToolsetName =
