Update ToolsetLibrary's list/get/set properties to respect access protections as we were allowing the llm to push buttons it really should not be able to. Add script error feedback when trying to perform an invalid operation.

#Synced-CL 51656843

diff --git a/Engine/Plugins/Experimental/ToolsetRegistry/Source/ToolsetRegistry/Private/Tests/ToolsetJsonTest.h b/Engine/Plugins/Experimental/ToolsetRegistry/Source/ToolsetRegistry/Private/Tests/ToolsetJsonTest.h
--- a/Engine/Plugins/Experimental/ToolsetRegistry/Source/ToolsetRegistry/Private/Tests/ToolsetJsonTest.h
+++ b/Engine/Plugins/Experimental/ToolsetRegistry/Source/ToolsetRegistry/Private/Tests/ToolsetJsonTest.h
@@ -19,9 +19,22 @@ public:
 	UPROPERTY(BlueprintReadWrite, Category="ToolsetTest")
 	bool VisibleProperty;
 
+	// No CPF_Edit/BlueprintVisible/BlueprintAssignable → AccessProtected for both get and set.
 	UPROPERTY()
 	bool HiddenProperty;
 
+	// CPF_Edit | CPF_EditConst → accessible for get, ReadOnly for set.
+	UPROPERTY(VisibleAnywhere, Category="ToolsetTest")
+	bool ReadOnlyProperty;
+
+	// CPF_Edit | CPF_DisableEditOnInstance → accessible for get, CannotEditInstance on non-templates.
+	UPROPERTY(EditDefaultsOnly, Category="ToolsetTest")
+	bool DefaultsOnlyProperty;
+
+	// CPF_Edit | CPF_DisableEditOnTemplate → accessible for get, CannotEditTemplate on templates (CDOs).
+	UPROPERTY(EditInstanceOnly, Category="ToolsetTest")
+	bool InstanceOnlyProperty;
+
 	UFUNCTION(BlueprintCallable, Category = "Main")
 	void TestDefault(
 		UObject* TestObject = nullptr, FColor TestColor = FColor(10, 20, 30, 40),
diff --git a/Engine/Plugins/Experimental/ToolsetRegistry/Source/ToolsetRegistry/Private/Tests/ToolsetLibraryTest.cpp b/Engine/Plugins/Experimental/ToolsetRegistry/Source/ToolsetRegistry/Private/Tests/ToolsetLibraryTest.cpp
--- a/Engine/Plugins/Experimental/ToolsetRegistry/Source/ToolsetRegistry/Private/Tests/ToolsetLibraryTest.cpp
+++ b/Engine/Plugins/Experimental/ToolsetRegistry/Source/ToolsetRegistry/Private/Tests/ToolsetLibraryTest.cpp
@@ -63,6 +63,28 @@ void FToolsetLibrarySpec::Define()
 		TestTrue("SetProp", TestObject->VisibleProperty);
 	});
 
+	It("GetObjectProperties omits missing properties from result.", [this]()
+	{
+		UToolsetJsonTestObject* TestObject = NewObject<UToolsetJsonTestObject>();
+		TestObject->VisibleProperty = true;
+		FString PropertiesJson = UToolsetLibrary::GetObjectProperties(
+			TestObject, { FName("visibleProperty"), FName("doesNotExist") });
+		TSharedPtr<FJsonObject> JsonObject = JsonStringToJsonObject(PropertiesJson);
+		TestTrue("Accessible property is present", JsonObject->HasField(TEXT("visibleProperty")));
+		TestFalse("Missing property is absent", JsonObject->HasField(TEXT("doesNotExist")));
+	});
+
+	It("GetObjectProperties omits access-protected properties from result.", [this]()
+	{
+		UToolsetJsonTestObject* TestObject = NewObject<UToolsetJsonTestObject>();
+		TestObject->VisibleProperty = true;
+		FString PropertiesJson = UToolsetLibrary::GetObjectProperties(
+			TestObject, { FName("visibleProperty"), FName("hiddenProperty") });
+		TSharedPtr<FJsonObject> JsonObject = JsonStringToJsonObject(PropertiesJson);
+		TestTrue("Accessible property is present", JsonObject->HasField(TEXT("visibleProperty")));
+		TestFalse("Protected property is absent", JsonObject->HasField(TEXT("hiddenProperty")));
+	});
+
 	It("Returns false when setting values on non-existent properties.", [this]()
 	{
 		UToolsetJsonTestObject* TestObject = NewObject<UToolsetJsonTestObject>();
@@ -72,6 +94,43 @@ void FToolsetLibrarySpec::Define()
 		TestFalse("SetProp", UToolsetLibrary::SetObjectProperties(TestObject, PropertiesJson));
 	});
 
+	It("Returns false when setting access-protected properties.", [this]()
+	{
+		UToolsetJsonTestObject* TestObject = NewObject<UToolsetJsonTestObject>();
+		TSharedPtr<FJsonObject> JsonObject = MakeShared<FJsonObject>();
+		JsonObject->SetBoolField(TEXT("hiddenProperty"), true);
+		TestFalse("SetProp", UToolsetLibrary::SetObjectProperties(TestObject, JsonToString(JsonObject.ToSharedRef())));
+		TestFalse("Property not modified", TestObject->HiddenProperty);
+	});
+
+	It("Returns false when setting read-only properties.", [this]()
+	{
+		UToolsetJsonTestObject* TestObject = NewObject<UToolsetJsonTestObject>();
+		TSharedPtr<FJsonObject> JsonObject = MakeShared<FJsonObject>();
+		JsonObject->SetBoolField(TEXT("readOnlyProperty"), true);
+		TestFalse("SetProp", UToolsetLibrary::SetObjectProperties(TestObject, JsonToString(JsonObject.ToSharedRef())));
+		TestFalse("Property not modified", TestObject->ReadOnlyProperty);
+	});
+
+	It("Returns false when setting defaults-only properties on instances.", [this]()
+	{
+		UToolsetJsonTestObject* TestObject = NewObject<UToolsetJsonTestObject>();
+		TSharedPtr<FJsonObject> JsonObject = MakeShared<FJsonObject>();
+		JsonObject->SetBoolField(TEXT("defaultsOnlyProperty"), true);
+		TestFalse("SetProp", UToolsetLibrary::SetObjectProperties(TestObject, JsonToString(JsonObject.ToSharedRef())));
+		TestFalse("Property not modified", TestObject->DefaultsOnlyProperty);
+	});
+
+	It("Returns false when setting instance-only properties on templates.", [this]()
+	{
+		UToolsetJsonTestObject* TestObject =
+			UToolsetJsonTestObject::StaticClass()->GetDefaultObject<UToolsetJsonTestObject>();
+		TSharedPtr<FJsonObject> JsonObject = MakeShared<FJsonObject>();
+		JsonObject->SetBoolField(TEXT("instanceOnlyProperty"), true);
+		TestFalse("SetProp", UToolsetLibrary::SetObjectProperties(TestObject, JsonToString(JsonObject.ToSharedRef())));
+		TestFalse("Property not modified", TestObject->InstanceOnlyProperty);
+	});
+
 	It("Gets subclasses of a class.", [this]()
 	{
 		TArray<UClass*> Subclasses = UToolsetLibrary::GetDerivedClasses(
diff --git a/Engine/Plugins/Experimental/ToolsetRegistry/Source/ToolsetRegistry/Private/ToolsetRegistry/ToolsetLibrary.cpp b/Engine/Plugins/Experimental/ToolsetRegistry/Source/ToolsetRegistry/Private/ToolsetRegistry/ToolsetLibrary.cpp
--- a/Engine/Plugins/Experimental/ToolsetRegistry/Source/ToolsetRegistry/Private/ToolsetRegistry/ToolsetLibrary.cpp
+++ b/Engine/Plugins/Experimental/ToolsetRegistry/Source/ToolsetRegistry/Private/ToolsetRegistry/ToolsetLibrary.cpp
@@ -8,6 +8,7 @@
 #include "UObject/UObjectIterator.h"
 #include "Misc/AssertionMacros.h"
 #include "UObject/PropertyAccessUtil.h"
+#include "Kismet/KismetSystemLibrary.h"
 #include "UnrealClient.h"
 
 #include "ToolsetRegistry/JsonConversion.h"
@@ -16,7 +17,7 @@
 
 #include UE_INLINE_GENERATED_CPP_BY_NAME(ToolsetLibrary)
 
-FString UToolsetLibrary::ListStructProperties(UStruct* Struct)
+FString UToolsetLibrary::ListStructProperties(UStruct* Struct, bool bUserVisiblePropertiesOnly)
 {
 	if (!Struct)
 	{
@@ -24,7 +25,7 @@ FString UToolsetLibrary::ListStructProperties(UStruct* Struct)
 	}
 	FString JsonProperties;
 	TSharedPtr<FJsonObject> ClassSchema =
-		UE::ToolsetRegistry::Internal::ToolsetJson::StructToJsonSchema(Struct, true);
+		UE::ToolsetRegistry::Internal::ToolsetJson::StructToJsonSchema(Struct, bUserVisiblePropertiesOnly);
 	const TSharedPtr<FJsonObject>* ClassProperties = nullptr;
 	if (ClassSchema->TryGetObjectField(TEXT("properties"), ClassProperties))
 	{
@@ -39,16 +40,32 @@ FString UToolsetLibrary::GetObjectProperties(const UObject* Object, const TArray
 {
 	check(Object);
 	TSharedPtr<FJsonObject> PropertyData = MakeShared<FJsonObject>();
+	TArray<FString> InaccessibleProperties;
 	for (const auto& PropertyName : PropertyNames)
 	{
-		FProperty* Property = Object->GetClass()->FindPropertyByName(PropertyName);
-		if(Property)
+		FProperty* Property = PropertyAccessUtil::FindPropertyByName(PropertyName, Object->GetClass());
+		if (!Property)
 		{
-			PropertyData->SetField(
-				PropertyName.ToString(),
-				UE::ToolsetRegistry::Internal::ToolsetJson::PropertyToJsonData(
-					Property, Property->ContainerPtrToValuePtr<void>(Object)));
+			InaccessibleProperties.Add(PropertyName.ToString());
+			continue;
 		}
+		const EPropertyAccessResultFlags AccessResult = PropertyAccessUtil::CanGetPropertyValue(Property);
+		if (EnumHasAnyFlags(AccessResult, EPropertyAccessResultFlags::PermissionDenied))
+		{
+			InaccessibleProperties.Add(PropertyName.ToString());
+			continue;
+		}
+		PropertyData->SetField(
+			PropertyName.ToString(),
+			UE::ToolsetRegistry::Internal::ToolsetJson::PropertyToJsonData(
+				Property, Property->ContainerPtrToValuePtr<void>(Object)));
+	}
+	if (!InaccessibleProperties.IsEmpty())
+	{
+		UKismetSystemLibrary::RaiseScriptError(FString::Printf(
+			TEXT("GetObjectProperties on '%s' (%s): the following properties could not be read: %s"),
+			*Object->GetPathName(), *Object->GetClass()->GetName(),
+			*FString::Join(InaccessibleProperties, TEXT(", "))));
 	}
 	return UE::ToolsetRegistry::Internal::JsonToString(PropertyData.ToSharedRef());
 }
@@ -60,34 +77,49 @@ bool UToolsetLibrary::SetObjectProperties(UObject* Object, const FString& Proper
 		return false;
 	}
 
-	bool AllPropertiesSet = false;
 	TSharedPtr<FJsonObject> PropertyData =
 		UE::ToolsetRegistry::Internal::JsonStringToJsonObject(PropertiesJson);
-	if (PropertyData.IsValid())
+	if (!PropertyData.IsValid())
 	{
-		AllPropertiesSet = true;
-		const UClass* ObjectClass = Object->GetClass();
-		for (const auto& Pair : PropertyData->Values)
+		return false;
+	}
+
+	const UClass* ObjectClass = Object->GetClass();
+	const bool bObjectIsTemplate = PropertyAccessUtil::IsObjectTemplate(Object);
+	TArray<FString> InaccessibleProperties;
+	for (const auto& Pair : PropertyData->Values)
+	{
+		const auto& Key = Pair.Key;
+		FProperty* Property = PropertyAccessUtil::FindPropertyByName(FName(Key), ObjectClass);
+		if (!Property)
 		{
-			const auto& Key = Pair.Key;
-			FProperty* Property = ObjectClass->FindPropertyByName(FName(Key));
-			if (Property)
-			{
-				TUniquePtr<FPropertyAccessChangeNotify> ChangeNotify = PropertyAccessUtil::BuildBasicChangeNotify(
-					Property, Object, EPropertyAccessChangeNotifyMode::Always, EPropertyChangeType::ValueSet);
-				PropertyAccessUtil::EmitPreChangeNotify(ChangeNotify.Get(), false);
-				void* SettingsPropPtr = Property->ContainerPtrToValuePtr<void>(Object);
-				UE::ToolsetRegistry::Internal::ToolsetJson::JsonDataToProperty(
-					PropertyData->TryGetField(Key), Property, SettingsPropPtr);
-				PropertyAccessUtil::EmitPostChangeNotify(ChangeNotify.Get(), false);
-			}
-			else
-			{
-				AllPropertiesSet = false;
-			}
+			InaccessibleProperties.Add(*Key);
+			continue;
 		}
+		const EPropertyAccessResultFlags AccessResult = PropertyAccessUtil::CanSetPropertyValue(
+			Property, PropertyAccessUtil::EditorReadOnlyFlags, bObjectIsTemplate);
+		if (EnumHasAnyFlags(AccessResult, EPropertyAccessResultFlags::PermissionDenied))
+		{
+			InaccessibleProperties.Add(*Key);
+			continue;
+		}
+		TUniquePtr<FPropertyAccessChangeNotify> ChangeNotify = PropertyAccessUtil::BuildBasicChangeNotify(
+			Property, Object, EPropertyAccessChangeNotifyMode::Always, EPropertyChangeType::ValueSet);
+		PropertyAccessUtil::EmitPreChangeNotify(ChangeNotify.Get(), false);
+		void* SettingsPropPtr = Property->ContainerPtrToValuePtr<void>(Object);
+		UE::ToolsetRegistry::Internal::ToolsetJson::JsonDataToProperty(
+			PropertyData->TryGetField(Key), Property, SettingsPropPtr);
+		PropertyAccessUtil::EmitPostChangeNotify(ChangeNotify.Get(), false);
+	}
+	if (!InaccessibleProperties.IsEmpty())
+	{
+		UKismetSystemLibrary::RaiseScriptError(FString::Printf(
+			TEXT("SetObjectProperties on '%s' (%s): the following properties could not be set: %s"),
+			*Object->GetPathName(), *Object->GetClass()->GetName(),
+			*FString::Join(InaccessibleProperties, TEXT(", "))));
+		return false;
 	}
-	return AllPropertiesSet;
+	return true;
 }
 
 TArray<UClass*> UToolsetLibrary::GetDerivedClasses(UClass* BaseClass)
diff --git a/Engine/Plugins/Experimental/ToolsetRegistry/Source/ToolsetRegistry/Private/ToolsetRegistry/ToolsetLibrary.h b/Engine/Plugins/Experimental/ToolsetRegistry/Source/ToolsetRegistry/Private/ToolsetRegistry/ToolsetLibrary.h
--- a/Engine/Plugins/Experimental/ToolsetRegistry/Source/ToolsetRegistry/Private/ToolsetRegistry/ToolsetLibrary.h
+++ b/Engine/Plugins/Experimental/ToolsetRegistry/Source/ToolsetRegistry/Private/ToolsetRegistry/ToolsetLibrary.h
@@ -16,12 +16,21 @@ class UToolsetLibrary : public UBlueprintFunctionLibrary
 	GENERATED_BODY()
 public:
 	/*
-	 * Returns the Blueprint visible properties in the class as JSON Schema.
+	 * Returns the properties of a struct as JSON Schema.
+	 * By default only returns user-visible (Blueprint-accessible) properties.
+	 *
+	 * WARNING: bUserVisiblePropertiesOnly=false bypasses all visibility filtering and returns
+	 * every non-deprecated property on the struct. Only use this when you know that all struct
+	 * properties should be exposed. For example, DataTable row structs where every field is a
+	 * data column regardless of its Blueprint visibility flags.
+	 *
 	 * @param Struct The struct to extract properties from.
+	 * @param bUserVisiblePropertiesOnly When true (default), only Blueprint-accessible properties
+	 *        are included. When false, all non-deprecated properties are returned.
 	 * @return A JSON Schema formatted string that describes the properties.
 	 */
 	UFUNCTION(BlueprintCallable, Category = "Toolset Library")
-	static FString ListStructProperties(UStruct* Struct);
+	static FString ListStructProperties(UStruct* Struct, bool bUserVisiblePropertiesOnly = true);
 
 	/*
 	 * Returns the values of the requested properties.
