[AI Assistant] Ensure the AI assistant and toolset registry are not packed with game builds.

#Synced-CL 51721957

diff --git a/Engine/Plugins/Experimental/AIAssistant/AIAssistant.uplugin b/Engine/Plugins/Experimental/AIAssistant/AIAssistant.uplugin
--- a/Engine/Plugins/Experimental/AIAssistant/AIAssistant.uplugin
+++ b/Engine/Plugins/Experimental/AIAssistant/AIAssistant.uplugin
@@ -5,6 +5,7 @@
 	"FriendlyName": "AI Assistant",
 	"Description": "",
 	"Category": "Other",
+	"EditorOnly": true,
 	"CreatedBy": "",
 	"CreatedByURL": "",
 	"DocsURL": "",
@@ -20,7 +21,8 @@
 		{
 			"Name": "AIAssistant",
 			"Type": "EditorNoCommandlet",
-			"LoadingPhase": "Default"
+			"LoadingPhase": "Default",
+			"TargetAllowList": ["Editor"]
 		}
 	],
 	"Plugins": [
diff --git a/Engine/Plugins/Experimental/ToolsetRegistry/ToolsetRegistry.uplugin b/Engine/Plugins/Experimental/ToolsetRegistry/ToolsetRegistry.uplugin
--- a/Engine/Plugins/Experimental/ToolsetRegistry/ToolsetRegistry.uplugin
+++ b/Engine/Plugins/Experimental/ToolsetRegistry/ToolsetRegistry.uplugin
@@ -5,6 +5,7 @@
 	"FriendlyName": "Toolset Registry",
 	"Description": "",
 	"Category": "Other",
+	"EditorOnly": true,
 	"CreatedBy": "",
 	"CreatedByURL": "",
 	"DocsURL": "",
@@ -20,7 +21,8 @@
 		{
 			"Name": "ToolsetRegistry",
 			"Type": "Editor",
-			"LoadingPhase": "Default"
+			"LoadingPhase": "Default",
+			"TargetAllowList": ["Editor"]
 		}
 	],
 	"Plugins": [
