[Toolset Registry] Add method to get all registered toolset schemas
Convenience method for MCP or other querants

#rnx
#Synced-CL 51757081

diff --git a/Engine/Plugins/Experimental/ToolsetRegistry/Source/ToolsetRegistry/Private/Tests/BlueprintLibraryToolsetTest.cpp b/Engine/Plugins/Experimental/ToolsetRegistry/Source/ToolsetRegistry/Private/Tests/BlueprintLibraryToolsetTest.cpp
--- a/Engine/Plugins/Experimental/ToolsetRegistry/Source/ToolsetRegistry/Private/Tests/BlueprintLibraryToolsetTest.cpp
+++ b/Engine/Plugins/Experimental/ToolsetRegistry/Source/ToolsetRegistry/Private/Tests/BlueprintLibraryToolsetTest.cpp
@@ -752,6 +752,29 @@ bool FBlueprintLibraryToolsetTestToolsetNamingNull::RunTest(
 	return TestTrue("Toolset name for null UClass should be empty", ToolsetName.IsEmpty());
 }
 
+DEFINE_SPEC(FToolsetRegistryToolsetTestGetAllSchemas,
+			"AI.ToolsetRegistry.BlueprintLibraryToolset.GetAllSchemas", ToolsetRegistryTest::Flags)
+
+void FToolsetRegistryToolsetTestGetAllSchemas::Define()
+{
+	Describe("UToolsetRegistry::GetAllToolsetJsonSchemas", [this]()
+	{
+		It("should return exactly what FToolsetRegistry::GetToolsetJsonSchemas does", [this]()
+		{
+			auto ToolsetRegistrySubsystem = UToolsetRegistrySubsystem::Get();
+			if (!ToolsetRegistrySubsystem.HasValue())
+			{
+				AddError(TEXT("Toolset Registry subsystem unavailable"));
+				return;
+			}
+			FString AllSchemas =
+				ToolsetRegistrySubsystem.GetValue()->ToolsetRegistry.GetToolsetJsonSchemas();
+			TestEqual(TEXT("Schemas text output"), AllSchemas,
+					  UToolsetRegistry::GetAllToolsetJsonSchemas());
+		});
+	});
+}
+
 // Run python tests
 IMPLEMENT_SIMPLE_AUTOMATION_TEST(
 	FBlueprintLibraryToolsetTestPythonTests,
diff --git a/Engine/Plugins/Experimental/ToolsetRegistry/Source/ToolsetRegistry/Private/ToolsetRegistry/UToolsetRegistry.cpp b/Engine/Plugins/Experimental/ToolsetRegistry/Source/ToolsetRegistry/Private/ToolsetRegistry/UToolsetRegistry.cpp
--- a/Engine/Plugins/Experimental/ToolsetRegistry/Source/ToolsetRegistry/Private/ToolsetRegistry/UToolsetRegistry.cpp
+++ b/Engine/Plugins/Experimental/ToolsetRegistry/Source/ToolsetRegistry/Private/ToolsetRegistry/UToolsetRegistry.cpp
@@ -149,3 +149,15 @@ FString UToolsetRegistry::GetToolsetJsonSchema(TSubclassOf<UToolsetDefinition> I
 		MakeShared<UE::ToolsetRegistry::FBlueprintLibraryToolset>(InToolsetClass);
 	return Toolset ? Toolset->GetToolsetJsonSchema() : FString();
 }
+
+FString UToolsetRegistry::GetAllToolsetJsonSchemas()
+{
+	auto ToolsetRegistrySubsystem = UToolsetRegistrySubsystem::Get(TEXT("GetAllToolsetJsonSchemas"));
+	FString SchemaString;
+	if (ToolsetRegistrySubsystem.HasValue())
+	{
+		SchemaString = ToolsetRegistrySubsystem.GetValue()->ToolsetRegistry.GetToolsetJsonSchemas();
+	}
+	return SchemaString;
+}
+
diff --git a/Engine/Plugins/Experimental/ToolsetRegistry/Source/ToolsetRegistry/Public/ToolsetRegistry/UToolsetRegistry.h b/Engine/Plugins/Experimental/ToolsetRegistry/Source/ToolsetRegistry/Public/ToolsetRegistry/UToolsetRegistry.h
--- a/Engine/Plugins/Experimental/ToolsetRegistry/Source/ToolsetRegistry/Public/ToolsetRegistry/UToolsetRegistry.h
+++ b/Engine/Plugins/Experimental/ToolsetRegistry/Source/ToolsetRegistry/Public/ToolsetRegistry/UToolsetRegistry.h
@@ -51,6 +51,10 @@ public:
 	// Meant for use while testing.
 	UFUNCTION(BlueprintPure, Category="ToolsetRegistry")
 	static UE_API FString GetToolsetJsonSchema(TSubclassOf<UToolsetDefinition> InToolsetClass);
+
+	// Get JSON schemas for all registered toolsets.
+	UFUNCTION(BlueprintPure, Category="ToolsetRegistry")
+	static UE_API FString GetAllToolsetJsonSchemas();
 };
 
 #undef UE_API
