Replace janky old viewport screenshot tool with a proper async tool that can look at both the current viewport and visual asset types. Removed temp logic that supported the old path and add tests for new tool.

#Synced-CL 51698457

diff --git a/Engine/Plugins/Experimental/ToolsetRegistry/Source/ToolsetRegistry/Private/Tests/EditorAppToolsetTest.cpp b/Engine/Plugins/Experimental/ToolsetRegistry/Source/ToolsetRegistry/Private/Tests/EditorAppToolsetTest.cpp
--- a/Engine/Plugins/Experimental/ToolsetRegistry/Source/ToolsetRegistry/Private/Tests/EditorAppToolsetTest.cpp
+++ b/Engine/Plugins/Experimental/ToolsetRegistry/Source/ToolsetRegistry/Private/Tests/EditorAppToolsetTest.cpp
@@ -3,8 +3,14 @@
 #include "CoreMinimal.h"
 #include "Misc/AutomationTest.h"
 
+#include "Editor.h"
+#include "Engine/World.h"
+#include "Misc/App.h"
+#include "Misc/CoreMiscDefines.h"
+#include "UObject/StrongObjectPtr.h"
 #include "ToolsetRegistry/EditorAppToolset.h"
 #include "ToolsetRegistry/JsonConversion.h"
+#include "ToolsetRegistry/ToolCallAsyncResultImage.h"
 
 
 #if WITH_DEV_AUTOMATION_TESTS
@@ -137,6 +143,101 @@ void FEditorAppToolsetSpec::Define()
 		TestTrue("CVar", CVarJsonData.IsValid());
 		TestEqual("CVar", CVarJsonData->Values.Num(), 0);
 	});
+
+	// CaptureAssetImage tests
+
+	It(TEXT("CaptureAssetImage returns an error for a non-existent asset"), [this]()
+	{
+		UToolCallAsyncResultImage* Result = UEditorAppToolset::CaptureAssetImage(
+			TEXT("/Game/NonExistent/DoesNotExist.DoesNotExist"));
+		if (TestNotNull(TEXT("Result"), Result) || !Result) return;
+		// Error is set synchronously before CaptureAssetImage returns.
+		TestTrue("Complete", Result->bIsComplete);
+		TestFalse("Error is set", Result->Error.IsEmpty());
+	});
+
+	LatentIt(TEXT("CaptureAssetImage renders a thumbnail for a static mesh"),
+		[this](const FDoneDelegate& Done)
+		{
+			if (!FApp::CanEverRender())
+			{
+				// Rendering is disabled; skipping thumbnail test.
+				Done.Execute();
+				return;
+			}
+
+			UToolCallAsyncResultImage* Result = UEditorAppToolset::CaptureAssetImage(
+				TEXT("/Engine/BasicShapes/Cube.Cube"));
+			TestNotNull("Result", Result);
+			if (!Result)
+			{
+				Done.Execute();
+				return;
+			}
+
+			TStrongObjectPtr<UToolCallAsyncResultImage> StrongResult(Result);
+
+			AddCommand(new FFunctionLatentCommand(
+				[this, Done, StrongResult]() mutable -> bool
+				{
+					if (!StrongResult->bIsComplete)
+					{
+						return false; // not done yet
+					}
+					TestTrue("No error", StrongResult->Error.IsEmpty());
+					TestFalse("Has image data", StrongResult->Value.Data.IsEmpty());
+					TestEqual("MIME type",
+						StrongResult->Value.MimeType, FString(TEXT("image/png")));
+					Done.Execute();
+					return true; // done
+				}));
+		});
+
+	LatentIt(TEXT("CaptureAssetImage screenshots the viewport for the current level"),
+		[this](const FDoneDelegate& Done)
+		{
+			if (!FApp::CanEverRender())
+			{
+				// Rendering is disabled; skipping thumbnail test.
+				Done.Execute();
+				return;
+			}
+
+			UWorld* EditorWorld = GEditor ? GEditor->GetEditorWorldContext().World() : nullptr;
+			if (!EditorWorld)
+			{
+				// No editor world available; skipping viewport screenshot test.
+				Done.Execute();
+				return;
+			}
+
+			FString CurrentLevelPath = EditorWorld->GetOutermost()->GetName();
+			UToolCallAsyncResultImage* Result =
+				UEditorAppToolset::CaptureAssetImage(CurrentLevelPath);
+			TestNotNull("Result", Result);
+			if (!Result)
+			{
+				Done.Execute();
+				return;
+			}
+
+			TStrongObjectPtr<UToolCallAsyncResultImage> StrongResult(Result);
+
+			AddCommand(new FFunctionLatentCommand(
+				[this, Done, StrongResult]() mutable -> bool
+				{
+					if (!StrongResult->bIsComplete)
+					{
+						return false; // not done yet
+					}
+					TestTrue("No error", StrongResult->Error.IsEmpty());
+					TestFalse("Has image data", StrongResult->Value.Data.IsEmpty());
+					TestEqual("MIME type",
+						StrongResult->Value.MimeType, FString(TEXT("image/png")));
+					Done.Execute();
+					return true; // done
+				}));
+		});
 }
 
 #endif
diff --git a/Engine/Plugins/Experimental/ToolsetRegistry/Source/ToolsetRegistry/Private/Tests/ToolCallAsyncResultImageTest.cpp b/Engine/Plugins/Experimental/ToolsetRegistry/Source/ToolsetRegistry/Private/Tests/ToolCallAsyncResultImageTest.cpp
new file mode 100644
--- /dev/null
+++ b/Engine/Plugins/Experimental/ToolsetRegistry/Source/ToolsetRegistry/Private/Tests/ToolCallAsyncResultImageTest.cpp
@@ -0,0 +1,29 @@
+// Copyright Epic Games, Inc. All Rights Reserved.
+
+#include "Dom/JsonObject.h"
+#include "Dom/JsonValue.h"
+#include "Templates/SharedPointer.h"
+#include "ToolsetRegistry/ToolCallAsyncResultImage.h"
+#include "ToolsetRegistry/ToolsetImage.h"
+#include "Tests/ToolCallAsyncResultTest.h"
+
+#if WITH_DEV_AUTOMATION_TESTS
+
+UE_TOOLSET_REGISTRY_TOOL_CALL_ASYNC_RESULT_WITH_VALUE_SPEC(
+	Image, UToolCallAsyncResultImage, FToolsetImage,
+	[]() -> FToolsetImage
+	{
+		FToolsetImage Image;
+		Image.MimeType = TEXT("image/png");
+		Image.Data = TEXT("abc123");
+		return Image;
+	},
+	[](const FToolsetImage& ExpectedValue) -> TSharedPtr<FJsonValue>
+	{
+		TSharedRef<FJsonObject> JsonObject = MakeShared<FJsonObject>();
+		JsonObject->SetStringField(TEXT("mimeType"), ExpectedValue.MimeType);
+		JsonObject->SetStringField(TEXT("data"), ExpectedValue.Data);
+		return MakeShared<FJsonValueObject>(JsonObject);
+	})
+
+#endif  // WITH_DEV_AUTOMATION_TESTS
diff --git a/Engine/Plugins/Experimental/ToolsetRegistry/Source/ToolsetRegistry/Private/ToolsetRegistry/EditorAppToolset.cpp b/Engine/Plugins/Experimental/ToolsetRegistry/Source/ToolsetRegistry/Private/ToolsetRegistry/EditorAppToolset.cpp
--- a/Engine/Plugins/Experimental/ToolsetRegistry/Source/ToolsetRegistry/Private/ToolsetRegistry/EditorAppToolset.cpp
+++ b/Engine/Plugins/Experimental/ToolsetRegistry/Source/ToolsetRegistry/Private/ToolsetRegistry/EditorAppToolset.cpp
@@ -2,13 +2,244 @@
 
 #include "EditorAppToolset.h"
 
+#include "AssetRegistry/IAssetRegistry.h"
+#include "Containers/Ticker.h"
+#include "Editor.h"
+#include "Engine/AssetManager.h"
+#include "Engine/StreamableManager.h"
+#include "Animation/AnimationAsset.h"
+#include "Animation/Skeleton.h"
+#include "Engine/SkeletalMesh.h"
+#include "Engine/StaticMesh.h"
+#include "Engine/Texture.h"
 #include "HAL/ConsoleManager.h"
 #include "Kismet/KismetSystemLibrary.h"
-
+#include "Materials/Material.h"
+#include "Materials/MaterialInterface.h"
+#include "Misc/ObjectThumbnail.h"
+#include "Misc/PackageName.h"
+#include "ObjectTools.h"
+#include "RenderingThread.h"
+#include "ShaderCompiler.h"
 #include "ToolsetRegistry/JsonConversion.h"
+#include "ToolsetRegistry/ToolCallAsyncResultImage.h"
+#include "ToolsetRegistry/ToolsetImage.h"
+#include "UObject/UObjectGlobals.h"
+#include "UnrealClient.h"
 
 #include UE_INLINE_GENERATED_CPP_BY_NAME(EditorAppToolset)
 
+namespace
+{
+	// Helper that captures a single viewport screenshot and resolves a result.
+	// The object is kept alive by a strong reference captured in the screenshot
+	// delegate lambda, which is released once the callback fires.
+	struct FViewportScreenshotCapture
+	{
+		FViewportScreenshotCapture()
+		{
+			Result = NewObject<UToolCallAsyncResultImage>();
+			Result->AddToRoot();
+		}
+
+		static UToolCallAsyncResultImage* Create()
+		{
+			TSharedRef<FViewportScreenshotCapture> Self = MakeShared<FViewportScreenshotCapture>();
+			Self->Handle = FScreenshotRequest::OnScreenshotCaptured().AddLambda(
+				[Self](int32 Width, int32 Height, const TArray<FColor>& Colors)
+				{
+					Self->OnCaptured(Width, Height, Colors);
+				});
+			FScreenshotRequest::RequestScreenshot(false);
+			return Self->Result;
+		}
+
+		~FViewportScreenshotCapture()
+		{
+			Result->RemoveFromRoot();
+		}
+
+	private:
+		void OnCaptured(int32 Width, int32 Height, const TArray<FColor>& Colors)
+		{
+			FScreenshotRequest::OnScreenshotCaptured().Remove(Handle);
+
+			const int32 MaxSize = 1280;
+			FIntPoint Dimensions(Width, Height);
+			TArray<FColor> ScaledPixels;
+			if (Width > MaxSize || Height > MaxSize)
+			{
+				const float Scale =
+					(Width >= Height) ? (float)MaxSize / Width : (float)MaxSize / Height;
+				Dimensions.X = FMath::Max(1, (int32)(Width * Scale));
+				Dimensions.Y = FMath::Max(1, (int32)(Height * Scale));
+				ScaledPixels.SetNumUninitialized(Dimensions.X * Dimensions.Y);
+				for (int32 Y = 0; Y < Dimensions.Y; ++Y)
+				{
+					for (int32 X = 0; X < Dimensions.X; ++X)
+					{
+						ScaledPixels[Y * Dimensions.X + X] =
+							Colors[(int32)(Y / Scale) * Width + (int32)(X / Scale)];
+					}
+				}
+			}
+
+			FToolsetImage Image;
+			const TArray<FColor>& Final = ScaledPixels.Num() > 0 ? ScaledPixels : Colors;
+			if (Image.SetFromBitmap(Final, Dimensions))
+			{
+				Result->SetValue(Image);
+			}
+			else
+			{
+				Result->SetError(TEXT("Failed to encode viewport screenshot."));
+			}
+		}
+
+		UToolCallAsyncResultImage* Result;
+		FDelegateHandle Handle;
+	};
+
+	// Helper that renders an asset thumbnail and resolves a result.
+	// Initiates an async load then polls each tick until the asset is loaded and
+	// all compilation (assets and shaders) is complete before rendering.
+	struct FAssetThumbnailCapture
+	{
+		explicit FAssetThumbnailCapture(const FString& InAssetPath)
+			: AssetPath(InAssetPath)
+			, StartTime(FPlatformTime::Seconds())
+		{
+			Result = NewObject<UToolCallAsyncResultImage>();
+			Result->AddToRoot();
+		}
+
+		static UToolCallAsyncResultImage* Create(const FString& InAssetPath)
+		{
+			TSharedRef<FAssetThumbnailCapture> Self = MakeShared<FAssetThumbnailCapture>(InAssetPath);
+			// Initiate async load; the ticker polls for completion.
+			Self->StreamableHandle = UAssetManager::Get().GetStreamableManager().RequestAsyncLoad(
+				FSoftObjectPath(Self->AssetPath), FStreamableDelegate());
+			Self->TickerHandle = FTSTicker::GetCoreTicker().AddTicker(
+				FTickerDelegate::CreateLambda([Self](float) -> bool
+				{
+					return Self->OnTick(Self);
+				}));
+			return Self->Result;
+		}
+
+		~FAssetThumbnailCapture()
+		{
+			Result->RemoveFromRoot();
+		}
+
+	private:
+		// Returns true to keep ticking, false when done.
+		bool OnTick(TSharedRef<FAssetThumbnailCapture> Self)
+		{
+			if (FPlatformTime::Seconds() - StartTime > 10.0)
+			{
+				Result->SetError(FString::Printf(TEXT("Timed out loading asset: %s"), *AssetPath));
+				return false;
+			}
+
+			// Check to see if asset is loaded.
+			TObjectPtr<UObject> Asset = FSoftObjectPath(AssetPath).ResolveObject();
+			if (!Asset)
+			{
+				return true;
+			}
+
+			// Poll until asset-level and shader compilation are both complete.
+			// This matches what the editor thumbnail system does before rendering.
+			const IInterface_AsyncCompilation* AsyncComp =
+				Cast<IInterface_AsyncCompilation>(Asset);
+			if (AsyncComp && AsyncComp->IsCompiling())
+			{
+				return true;
+			}
+
+			// UMaterialInstance::IsCompiling() only checks the instance's own static
+			// permutation resources; check the base material separately.
+			if (UMaterialInterface* MatInterface = Cast<UMaterialInterface>(Asset))
+			{
+				if (MatInterface->IsCompiling())
+				{
+					return true;
+				}
+				if (UMaterial* Base = MatInterface->GetMaterial(); Base && Base->IsCompiling())
+				{
+					return true;
+				}
+			}
+			// Poll for texture async build and mip streaming completion.
+			if (UTexture* Texture = Cast<UTexture>(Asset))
+			{
+				if (!Texture->IsAsyncCacheComplete() || Texture->HasPendingInitOrStreaming())
+				{
+					return true;
+				}
+			}
+
+			// All compilation and streaming done. Release the streamable handle and render.
+			StreamableHandle.Reset();
+			TickerHandle.Reset();
+			FlushRenderingCommands();
+
+			FObjectThumbnail NewThumbnail;
+			ThumbnailTools::RenderThumbnail(
+				Asset,
+				ThumbnailTools::DefaultThumbnailSize,
+				ThumbnailTools::DefaultThumbnailSize,
+				ThumbnailTools::EThumbnailTextureFlushMode::NeverFlush,
+				nullptr,
+				&NewThumbnail);
+
+			const int32 Width = NewThumbnail.GetImageWidth();
+			const int32 Height = NewThumbnail.GetImageHeight();
+			const TArray<uint8>& ImageBytes = NewThumbnail.GetUncompressedImageData();
+			const int32 ExpectedBytes = Width * Height * 4; // BGRA8, 4 bytes per pixel
+
+			if (Width == 0 || Height == 0 || ImageBytes.Num() != ExpectedBytes)
+			{
+				Result->SetError(FString::Printf(
+					TEXT("Failed to render thumbnail for: %s"), *Asset->GetPathName()));
+				return false;
+			}
+
+			// The thumbnail data is raw BGRA8. Copy into FColor and swap R and B so the
+			// memory layout is RGBA, then encode with ERGBFormat::RGBA (no libpng channel
+			// transform needed). This avoids relying on PNG_TRANSFORM_BGR being compiled in.
+			// Force alpha to 255 since some thumbnail renderers do not write the alpha channel.
+			TArray<FColor> Colors;
+			Colors.SetNumUninitialized(Width * Height);
+			FMemory::Memcpy(Colors.GetData(), ImageBytes.GetData(), ImageBytes.Num());
+			for (FColor& Color : Colors)
+			{
+				Swap(Color.R, Color.B);
+				Color.A = 255;
+			}
+
+			FToolsetImage Image;
+			if (Image.SetFromBitmap(Colors, FIntPoint(Width, Height), ERGBFormat::RGBA))
+			{
+				Result->SetValue(Image);
+			}
+			else
+			{
+				Result->SetError(FString::Printf(
+					TEXT("Failed to encode thumbnail for: %s"), *Asset->GetPathName()));
+			}
+			return false;
+		}
+
+		FString AssetPath;
+		UToolCallAsyncResultImage* Result;
+		TSharedPtr<FStreamableHandle> StreamableHandle;
+		FTSTicker::FDelegateHandle TickerHandle;
+		double StartTime;
+	};
+}
+
 FString UEditorAppToolset::SearchCVars(const FString& Name)
 {
 	TSharedPtr<FJsonObject> Results = MakeShared<FJsonObject>();
@@ -62,3 +293,58 @@ TSharedPtr<FJsonObject> UEditorAppToolset::CVarToJson(IConsoleObject* CVar)
 	}
 	return CVarData;
 }
+
+UToolCallAsyncResultImage* UEditorAppToolset::CaptureAssetImage(const FString& AssetPath)
+{
+	// An empty path or the currently loaded level to viewport screenshot.
+	UWorld* EditorWorld = GEditor ? GEditor->GetEditorWorldContext().World() : nullptr;
+	FString CurrentLevelPackage = EditorWorld
+		? EditorWorld->GetOutermost()->GetName()
+		: TEXT("");
+	if (AssetPath.IsEmpty() ||
+		CurrentLevelPackage == FPackageName::ObjectPathToPackageName(AssetPath))
+	{
+		return FViewportScreenshotCapture::Create();
+	}
+
+	// Thumbnail rendering path
+	TObjectPtr<UToolCallAsyncResultImage> Result = NewObject<UToolCallAsyncResultImage>();
+	// Normalize a bare package path to a full object path so the registry lookup
+	// and async load both work. e.g. /Game/Foo/Bar  ->  /Game/Foo/Bar.Bar
+	FString ObjectPath = AssetPath;
+	if (!ObjectPath.Contains(TEXT(".")))
+	{
+		int32 LastSlash;
+		if (ObjectPath.FindLastChar(TEXT('/'), LastSlash))
+		{
+			ObjectPath += TEXT(".") + ObjectPath.RightChop(LastSlash + 1);
+		}
+	}
+	const FAssetData AssetData =
+		IAssetRegistry::GetChecked().GetAssetByObjectPath(FSoftObjectPath(ObjectPath));
+	if (!AssetData.IsValid())
+	{
+		Result->SetError(FString::Printf(TEXT("Asset not found: %s"), *ObjectPath));
+		return Result.Get();
+	}
+	else if (AssetData.IsInstanceOf<UWorld>())
+	{
+		Result->SetError(FString::Printf(
+			TEXT("Cannot capture a level that is not currently loaded: %s"), *ObjectPath));
+	}
+	else if (!AssetData.IsInstanceOf<UAnimationAsset>() &&
+		!AssetData.IsInstanceOf<USkeleton>() &&
+		!AssetData.IsInstanceOf<UStaticMesh>() &&
+		!AssetData.IsInstanceOf<USkeletalMesh>() &&
+		!AssetData.IsInstanceOf<UTexture>() &&
+		!AssetData.IsInstanceOf<UMaterialInterface>())
+	{
+		Result->SetError(FString::Printf(
+			TEXT("Asset type does not support image capture: %s"), *ObjectPath));
+	}
+	else
+	{
+		Result = FAssetThumbnailCapture::Create(ObjectPath);
+	}
+	return Result.Get();
+}
diff --git a/Engine/Plugins/Experimental/ToolsetRegistry/Source/ToolsetRegistry/Private/ToolsetRegistry/EditorAppToolset.h b/Engine/Plugins/Experimental/ToolsetRegistry/Source/ToolsetRegistry/Private/ToolsetRegistry/EditorAppToolset.h
--- a/Engine/Plugins/Experimental/ToolsetRegistry/Source/ToolsetRegistry/Private/ToolsetRegistry/EditorAppToolset.h
+++ b/Engine/Plugins/Experimental/ToolsetRegistry/Source/ToolsetRegistry/Private/ToolsetRegistry/EditorAppToolset.h
@@ -10,6 +10,8 @@
 
 class FJsonObject;
 class IConsoleObject;
+class UObject;
+class UToolCallAsyncResultImage;
 
 /// Contains tools that query or modify the state of the Unreal Editor application.
 UCLASS(BlueprintType, MinimalAPI)
@@ -25,6 +27,17 @@ public:
 	UFUNCTION(BlueprintCallable, Category = "EditorAppToolset")
 	static FString SearchCVars(const FString& Name);
 
+	/*
+	 * Generates an image of an asset. For the current level, a viewport screenshot is taken.
+	 * For all other assets (e.g. static meshes, skeletal meshes, skeletons, animations,
+	 * montages, materials, textures), a thumbnail is rendered.
+	 * @param AssetPath The path to the asset, e.g. '/Game/Meshes/SM_Cube'.
+	 *   Pass empty string or the current level path to capture a viewport screenshot instead.
+	 * @return An image of the asset.
+	 */
+	UFUNCTION(BlueprintCallable, Category = "EditorAppToolset")
+	static UToolCallAsyncResultImage* CaptureAssetImage(const FString& AssetPath);
+
 private:
 	static TSharedPtr<FJsonObject> CVarToJson(IConsoleObject* CVar);
 
diff --git a/Engine/Plugins/Experimental/ToolsetRegistry/Source/ToolsetRegistry/Private/ToolsetRegistry/ToolsetImage.cpp b/Engine/Plugins/Experimental/ToolsetRegistry/Source/ToolsetRegistry/Private/ToolsetRegistry/ToolsetImage.cpp
--- a/Engine/Plugins/Experimental/ToolsetRegistry/Source/ToolsetRegistry/Private/ToolsetRegistry/ToolsetImage.cpp
+++ b/Engine/Plugins/Experimental/ToolsetRegistry/Source/ToolsetRegistry/Private/ToolsetRegistry/ToolsetImage.cpp
@@ -11,21 +11,21 @@
 #include UE_INLINE_GENERATED_CPP_BY_NAME(ToolsetImage)
 
 void EncodeImage(const TSharedPtr<IImageWrapper>& ImageWrapper, FString& OutData)
-{	
+{
 	const TArray64<uint8>& CompressedData = ImageWrapper->GetCompressed();
 	OutData = FBase64::Encode(CompressedData.GetData(), CompressedData.Num());
 }
 
-bool FToolsetImage::SetFromBitmap(const TArray<FColor>& Bitmap, FIntPoint Dimensions)
+bool FToolsetImage::SetFromBitmap(const TArray<FColor>& Bitmap, FIntPoint Dimensions, ERGBFormat Format)
 {
 	IImageWrapperModule& ImageWrapperModule = FModuleManager::LoadModuleChecked<IImageWrapperModule>("ImageWrapper");
 	TSharedPtr<IImageWrapper> ImageWrapper = ImageWrapperModule.CreateImageWrapper(EImageFormat::PNG);
 
-	if (ImageWrapper.IsValid() && 
-		ImageWrapper->SetRaw(Bitmap.GetData(), Bitmap.Num() * sizeof(FColor), Dimensions.X, Dimensions.Y, ERGBFormat::BGRA, 8))
+	if (ImageWrapper.IsValid() &&
+		ImageWrapper->SetRaw(Bitmap.GetData(), Bitmap.Num() * sizeof(FColor), Dimensions.X, Dimensions.Y, Format, 8))
 	{
 		EncodeImage(ImageWrapper, Data);
-		MimeType = "image/png";
+		MimeType = TEXT("image/png");
 		return true;
 	}
 
@@ -41,7 +41,7 @@ bool FToolsetImage::SetFromFile(const FString& Path)
 	}
 
 	IImageWrapperModule& ImageWrapperModule = FModuleManager::LoadModuleChecked<IImageWrapperModule>("ImageWrapper");
-	
+
 	EImageFormat ImageFormat = ImageWrapperModule.DetectImageFormat(ImageData.GetData(), ImageData.Num());
 	if (ImageFormat != EImageFormat::PNG)
 	{
@@ -55,6 +55,6 @@ bool FToolsetImage::SetFromFile(const FString& Path)
 	}
 
 	EncodeImage(ImageWrapper, Data);
-	MimeType = "image/png";
+	MimeType = TEXT("image/png");
 	return true;
-}
\ No newline at end of file
+}
diff --git a/Engine/Plugins/Experimental/ToolsetRegistry/Source/ToolsetRegistry/Private/ToolsetRegistry/ToolsetLibrary.cpp b/Engine/Plugins/Experimental/ToolsetRegistry/Source/ToolsetRegistry/Private/ToolsetRegistry/ToolsetLibrary.cpp
--- a/Engine/Plugins/Experimental/ToolsetRegistry/Source/ToolsetRegistry/Private/ToolsetRegistry/ToolsetLibrary.cpp
+++ b/Engine/Plugins/Experimental/ToolsetRegistry/Source/ToolsetRegistry/Private/ToolsetRegistry/ToolsetLibrary.cpp
@@ -170,17 +170,3 @@ TArray<UScriptStruct*> UToolsetLibrary::GetDerivedStructs(UScriptStruct* BaseStr
 	}
 	return DerivedStructs;
 }
-
-FString UToolsetLibraryTemp::ScreenshotPath(const FString& Filename)
-{
-	FString FullPath = Filename;
-	FScreenshotRequest::CreateViewportScreenShotFilename(FullPath);
-	return FullPath;
-}
-
-FToolsetImage UToolsetLibraryTemp::ScreenshotToImage(const FString& FullPath)
-{	
-	FToolsetImage Image;
-	Image.SetFromFile(FullPath);
-	return Image;
-}
diff --git a/Engine/Plugins/Experimental/ToolsetRegistry/Source/ToolsetRegistry/Private/ToolsetRegistry/ToolsetLibrary.h b/Engine/Plugins/Experimental/ToolsetRegistry/Source/ToolsetRegistry/Private/ToolsetRegistry/ToolsetLibrary.h
--- a/Engine/Plugins/Experimental/ToolsetRegistry/Source/ToolsetRegistry/Private/ToolsetRegistry/ToolsetLibrary.h
+++ b/Engine/Plugins/Experimental/ToolsetRegistry/Source/ToolsetRegistry/Private/ToolsetRegistry/ToolsetLibrary.h
@@ -68,22 +68,3 @@ public:
 	UFUNCTION(BlueprintCallable, Category = "Toolset Library")
 	static TArray<UScriptStruct*> GetDerivedStructs(UScriptStruct* BaseStruct);
 };
-
-// Provides workarounds for the lack of async tool call support. Once we have that,
-// we should remove these functions and implement a tool that directly takes a screenshot and
-// returns a ToolsetImage.
-UCLASS(BlueprintType, MinimalAPI)
-class UToolsetLibraryTemp : public UBlueprintFunctionLibrary
-{
-	GENERATED_BODY()
-public:
-
-	// Converts the file name of a screenshot into a full path based on UEs default logic.
-	UFUNCTION(BlueprintCallable, Category = "Toolset Library Temp")
-	static FString ScreenshotPath(const FString& Filename);
-
-	// Converts a screenshot image on disk to a ToolsetImage.
-	UFUNCTION(BlueprintCallable, Category = "Toolset Library Temp")
-	static FToolsetImage ScreenshotToImage(const FString& FullPath);
-};
-
diff --git a/Engine/Plugins/Experimental/ToolsetRegistry/Source/ToolsetRegistry/Public/ToolsetRegistry/ToolCallAsyncResultImage.h b/Engine/Plugins/Experimental/ToolsetRegistry/Source/ToolsetRegistry/Public/ToolsetRegistry/ToolCallAsyncResultImage.h
new file mode 100644
--- /dev/null
+++ b/Engine/Plugins/Experimental/ToolsetRegistry/Source/ToolsetRegistry/Public/ToolsetRegistry/ToolCallAsyncResultImage.h
@@ -0,0 +1,28 @@
+// Copyright Epic Games, Inc. All Rights Reserved.
+
+#pragma once
+
+#include "ToolsetRegistry/ToolCallAsyncResult.h"
+#include "ToolsetRegistry/ToolsetImage.h"
+
+#include "ToolCallAsyncResultImage.generated.h"
+
+// Async tool call result that completes with a ToolsetImage.
+UCLASS(BlueprintType)
+class UToolCallAsyncResultImage : public UToolCallAsyncResult
+{
+	GENERATED_BODY()
+
+public:
+	// Set the result.
+	UFUNCTION(BlueprintCallable, Category = "ToolsetRegistry")
+	bool SetValue(const FToolsetImage& InValue)
+	{
+		return MaybeBroadcastSuccessfulCompletion(FToolsetImage(InValue), Value);
+	}
+
+	// Value of the result.
+	UPROPERTY(BlueprintReadOnly, Category = "ToolsetRegistry")
+	FToolsetImage Value;
+};
+
diff --git a/Engine/Plugins/Experimental/ToolsetRegistry/Source/ToolsetRegistry/Public/ToolsetRegistry/ToolsetImage.h b/Engine/Plugins/Experimental/ToolsetRegistry/Source/ToolsetRegistry/Public/ToolsetRegistry/ToolsetImage.h
--- a/Engine/Plugins/Experimental/ToolsetRegistry/Source/ToolsetRegistry/Public/ToolsetRegistry/ToolsetImage.h
+++ b/Engine/Plugins/Experimental/ToolsetRegistry/Source/ToolsetRegistry/Public/ToolsetRegistry/ToolsetImage.h
@@ -1,8 +1,9 @@
-// Copyright Epic Games, Inc. All Rights Reserved. 
+// Copyright Epic Games, Inc. All Rights Reserved.
 
 #pragma once
 
 #include "CoreMinimal.h"
+#include "IImageWrapper.h"
 
 #include "ToolsetImage.generated.h"
 
@@ -21,9 +22,11 @@ public:
 	/// The image data encoded as a base64 string.
 	UPROPERTY(BlueprintReadOnly, Category = "Image")
 	FString Data;
-	
+
 	/// Populates the ToolsetImage based on a raw color image.
-	UE_API bool SetFromBitmap(const TArray<FColor>& Bitmap, FIntPoint Dimensions);
+	/// Format controls the channel layout of Bitmap in memory (default BGRA matches FColor on LE).
+	UE_API bool SetFromBitmap(const TArray<FColor>& Bitmap, FIntPoint Dimensions,
+		ERGBFormat Format = ERGBFormat::BGRA);
 
 	/// Populates the ToolsetImage from a file saved on disk.
 	UE_API bool SetFromFile(const FString& Path);
diff --git a/Engine/Plugins/Experimental/ToolsetRegistry/Source/ToolsetRegistry/ToolsetRegistry.Build.cs b/Engine/Plugins/Experimental/ToolsetRegistry/Source/ToolsetRegistry/ToolsetRegistry.Build.cs
--- a/Engine/Plugins/Experimental/ToolsetRegistry/Source/ToolsetRegistry/ToolsetRegistry.Build.cs
+++ b/Engine/Plugins/Experimental/ToolsetRegistry/Source/ToolsetRegistry/ToolsetRegistry.Build.cs
@@ -12,32 +12,33 @@ public class ToolsetRegistry : ModuleRules
 		bUseUnity = true; // ..toggle as needed, use false for IWYU audit
 
 
-		PublicDefinitions.Add("WITH_AIASSISTANT_EPIC_INTERNAL=1"); 
+		PublicDefinitions.Add("WITH_AIASSISTANT_EPIC_INTERNAL=1");
+
 
-		
 		PublicIncludePaths.AddRange(
 			new string[] {
 			}
 			);
-				
-		
+
+
 		PrivateIncludePaths.AddRange(
 			new string[] {
 			}
 			);
-			
-		
+
+
 		PublicDependencyModuleNames.AddRange(
 			new string[]
 			{
 				"Core",
 			}
 			);
-			
-		
+
+
 		PrivateDependencyModuleNames.AddRange(
 			new string[]
 			{
+				"AssetRegistry",
 				"BlueprintGraph",
 				"CoreUObject",
 				"DeveloperSettings",
@@ -52,12 +53,14 @@ public class ToolsetRegistry : ModuleRules
 				"KismetCompiler",
 				"Projects",
 				"PythonScriptPlugin",
+				"RenderCore",
+				"RHI",
 				"StatusBar",
 				"UnrealEd",
 			}
 			);
-		
-		
+
+
 		DynamicallyLoadedModuleNames.AddRange(
 			new string[]
 			{
