[Toolset Registry] Add object validation to tool call param JSON conversion

#rnx
#Synced-CL 51865821

diff --git a/Engine/Plugins/Experimental/ToolsetRegistry/Source/ToolsetRegistry/Private/Tests/BlueprintLibraryToolsetTest.cpp b/Engine/Plugins/Experimental/ToolsetRegistry/Source/ToolsetRegistry/Private/Tests/BlueprintLibraryToolsetTest.cpp
--- a/Engine/Plugins/Experimental/ToolsetRegistry/Source/ToolsetRegistry/Private/Tests/BlueprintLibraryToolsetTest.cpp
+++ b/Engine/Plugins/Experimental/ToolsetRegistry/Source/ToolsetRegistry/Private/Tests/BlueprintLibraryToolsetTest.cpp
@@ -670,6 +670,33 @@ bool FBlueprintLibraryToolsetTestExecuteCppToolWithError::RunTest(
 	return true;
 }
 
+// Invoke a C++ tool that produces errors from invalid object param checks.
+IMPLEMENT_SIMPLE_AUTOMATION_TEST(
+	FBlueprintLibraryToolsetTestExecuteCppToolWithBadParam,
+	"AI.ToolsetRegistry.BlueprintLibraryToolset.ExecuteCppToolWithBadParam",
+	ToolsetRegistryTest::Flags)
+
+bool FBlueprintLibraryToolsetTestExecuteCppToolWithBadParam::RunTest(
+	const FString& UnusedParameters)
+{
+	ToolsetRegistryContainer Container(UFakeToolsetWithObjectParams::StaticClass());
+	TSharedPtr<UE::ToolsetRegistry::FBlueprintLibraryToolset> Toolset = Container.ToolsetHandler();
+	if (!TestTrue("Valid toolset handler returned", Toolset.IsValid()))
+	{
+		return false;
+	}
+
+	FString Result;
+	FString Error;
+	UToolsetRegistry::ExecuteTool(
+		Toolset->GetToolsetName(), TEXT("TestObjectParam"),
+		TEXT(R"({"testObject": {"refPath": "/Invalid/Path/DoesNotExist"}})"), Result, Error);
+	TestFalse("Tool call with invalid object param error message", Error.IsEmpty());
+	TestTrue("Tool call with invalid object param returns correct raised exception",
+			 Error.Contains(TEXT("is not valid")));
+	return true;
+}
+
 // Invoke python tool that produces errors.
 IMPLEMENT_SIMPLE_AUTOMATION_TEST(
 	FBlueprintLibraryToolsetTestExecutePythonToolWithError,
diff --git a/Engine/Plugins/Experimental/ToolsetRegistry/Source/ToolsetRegistry/Private/Tests/BlueprintLibraryToolsetTest.h b/Engine/Plugins/Experimental/ToolsetRegistry/Source/ToolsetRegistry/Private/Tests/BlueprintLibraryToolsetTest.h
--- a/Engine/Plugins/Experimental/ToolsetRegistry/Source/ToolsetRegistry/Private/Tests/BlueprintLibraryToolsetTest.h
+++ b/Engine/Plugins/Experimental/ToolsetRegistry/Source/ToolsetRegistry/Private/Tests/BlueprintLibraryToolsetTest.h
@@ -169,3 +169,19 @@ public:
 		UKismetSystemLibrary::RaiseScriptError(ErrorMessage);
 	}
 };
+
+/** Fake Toolset that with methods taking object and class params. */
+UCLASS(Blueprintable, Hidden)
+class UFakeToolsetWithObjectParams : public UToolsetDefinition
+{
+	GENERATED_BODY()
+
+public:
+
+	UFUNCTION(BlueprintCallable, Category = Fake)
+	static void TestObjectParam(UObject* TestObject) {}
+
+	UFUNCTION(BlueprintCallable, Category = Fake)
+	static void TestClassParam(UClass* TestClass) {}
+};
+
diff --git a/Engine/Plugins/Experimental/ToolsetRegistry/Source/ToolsetRegistry/Private/Tests/ToolCallScriptExceptionHandlerTest.cpp b/Engine/Plugins/Experimental/ToolsetRegistry/Source/ToolsetRegistry/Private/Tests/ToolCallScriptExceptionHandlerTest.cpp
new file mode 100644
--- /dev/null
+++ b/Engine/Plugins/Experimental/ToolsetRegistry/Source/ToolsetRegistry/Private/Tests/ToolCallScriptExceptionHandlerTest.cpp
@@ -0,0 +1,126 @@
+// Copyright Epic Games, Inc. All Rights Reserved.
+
+
+#include "Misc/AutomationTest.h"
+
+#include "Tests/ToolsetJsonTest.h"
+#include "Tests/ToolCallTestHelpers.h"
+#include "Tests/ToolsetRegistryTestFlags.h"
+#include "ToolsetRegistry/JsonValueOrError.h"
+#include "ToolsetRegistry/ObjectFunctionToolCall.h"
+#include "ToolsetRegistry/ToolCallScriptExceptionHandler.h"
+
+PRAGMA_DISABLE_DEPRECATION_WARNINGS
+
+#if WITH_DEV_AUTOMATION_TESTS
+
+using namespace UE::ToolsetRegistry;
+
+BEGIN_DEFINE_SPEC(FToolCallScriptExceptionHandlerSpec,
+	"AI.ToolsetRegistry.ToolCallScriptExceptionHandler",
+	ToolsetRegistryTest::Flags)
+END_DEFINE_SPEC(FToolCallScriptExceptionHandlerSpec)
+
+void FToolCallScriptExceptionHandlerSpec::Define()
+{
+	It("Accumulates multiple parameter errors with newlines", [this]()
+	{
+		TObjectPtr<UToolsetJsonTestObject> TestObject = NewObject<UToolsetJsonTestObject>();
+
+		// Pass two invalid references to trigger multiple RaiseScriptError calls
+		FString InvalidJson = FString::Printf(
+			TEXT(R"({"firstObject": {"refPath": "%s"}, "secondObject": {"refPath": "%s"}})"),
+			*UObject::StaticClass()->GetPathName(),  // Wrong type for FirstObject
+			*UClass::StaticClass()->GetPathName());  // Wrong type for SecondObject
+
+		FJsonValueOrError Result = TestHelpers::ExecuteToolCallWithJson(
+			TestObject, TEXT("TestMultipleRefs"), InvalidJson);
+
+		if (!TestTrue("Should have error", Result.HasError())) return;
+
+		FString ErrorMessage = Result.GetError();
+
+		// Verify both errors are present
+		TestTrue("Contains first param type validation error",
+				 ErrorMessage.Contains(TEXT("is not valid ToolsetJsonTestObject for property 'FirstObject'")));
+		TestTrue("Contains second param type validation error",
+				 ErrorMessage.Contains(TEXT("is not valid ToolsetJsonTestObject for property 'SecondObject'")));
+
+		// Verify errors are accumulated with newlines
+		TestTrue("Contains newline separator",
+				 ErrorMessage.Contains(TEXT("\n")));
+
+		// Verify parameter error prefix
+		TestTrue("Contains parameter error prefix",
+				 ErrorMessage.Contains(TEXT("Parameter error")));
+	});
+
+	It("Captures single parameter error correctly", [this]()
+	{
+		TObjectPtr<UToolsetJsonTestObject> TestObject = NewObject<UToolsetJsonTestObject>();
+
+		// Pass one invalid reference
+		FString InvalidJson = FString::Printf(
+			TEXT(R"({"testObject": {"refPath": "%s"}})"),
+			*UObject::StaticClass()->GetPathName());  // Wrong type
+
+		FJsonValueOrError Result = TestHelpers::ExecuteToolCallWithJson(
+			TestObject, TEXT("TestSpecObjectParam"), InvalidJson);
+
+		if (!TestTrue("Should have error", Result.HasError())) return;
+
+		FString ErrorMessage = Result.GetError();
+
+		// Verify error message format
+		TestTrue("Contains type validation error",
+				 ErrorMessage.Contains(TEXT("is not valid ToolsetJsonTestObject")));
+		TestTrue("Error has parameter prefix",
+				 ErrorMessage.Contains(TEXT("Parameter error")));
+	});
+
+	It("Returns no error when all parameters are valid", [this]()
+	{
+		TObjectPtr<UToolsetJsonTestObject> TestObject = NewObject<UToolsetJsonTestObject>();
+		TObjectPtr<UToolsetJsonTestObject> ValidObject1 = NewObject<UToolsetJsonTestObject>();
+		TObjectPtr<UToolsetJsonTestObject> ValidObject2 = NewObject<UToolsetJsonTestObject>();
+
+		// Pass valid references
+		FString ValidJson = FString::Printf(
+			TEXT(R"({"firstObject": {"refPath": "%s"}, "secondObject": {"refPath": "%s"}})"),
+			*ValidObject1->GetPathName(),
+			*ValidObject2->GetPathName());
+
+		FJsonValueOrError Result = TestHelpers::ExecuteToolCallWithJson(
+			TestObject, TEXT("TestMultipleRefs"), ValidJson);
+
+		// Should succeed with no error
+		TestFalse("Should not have error", Result.HasError());
+	});
+
+	It("Returns a fallback message when an empty error message is raised", [this]()
+	{
+		TObjectPtr<UToolsetJsonTestObject> TestObject = NewObject<UToolsetJsonTestObject>();
+
+		UFunction* Function = TestObject->GetClass()->FindFunctionByName(TEXT("TestEmptyError"));
+		if (!TestNotNull("Function exists", Function)) return;
+
+		auto ToolCall = MakeShared<FObjectFunctionToolCall>(TestObject, Function);
+		TSharedPtr<UE::ToolsetRegistry::Internal::FToolCallScriptExceptionHandler> ExceptionHandler =
+			MakeShared<UE::ToolsetRegistry::Internal::FToolCallScriptExceptionHandler>();
+
+		ToolCall->Execute(
+			FObjectFunctionToolCall::FFunctionInputParamsJson(TInPlaceType<FString>(), FString()),
+			ExceptionHandler);
+
+		FString ErrorMessage = ExceptionHandler->GetException();
+
+		if (TestFalse("Should have error", ErrorMessage.IsEmpty())) return;
+
+		// Verify error message format
+		TestTrue("Error has fallback format", ErrorMessage.Contains(TEXT("UserRaisedError in ")));
+	});
+}
+
+#endif
+
+PRAGMA_ENABLE_DEPRECATION_WARNINGS
diff --git a/Engine/Plugins/Experimental/ToolsetRegistry/Source/ToolsetRegistry/Private/Tests/ToolCallTestHelpers.h b/Engine/Plugins/Experimental/ToolsetRegistry/Source/ToolsetRegistry/Private/Tests/ToolCallTestHelpers.h
new file mode 100644
--- /dev/null
+++ b/Engine/Plugins/Experimental/ToolsetRegistry/Source/ToolsetRegistry/Private/Tests/ToolCallTestHelpers.h
@@ -0,0 +1,38 @@
+// Copyright Epic Games, Inc. All Rights Reserved.
+
+#pragma once
+
+#include "CoreMinimal.h"
+#include "Templates/ValueOrError.h"
+#include "ToolsetRegistry/JsonValueOrError.h"
+#include "ToolsetRegistry/ObjectFunctionToolCall.h"
+#include "ToolsetRegistry/ToolCallScriptExceptionHandler.h"
+
+namespace UE::ToolsetRegistry::TestHelpers
+{
+	// Helper to execute a tool call for parameter validation.  Instantiates a tool call script exception
+	// handler to capture errors raised during parameter validation before invoking the function.
+	inline FJsonValueOrError ExecuteToolCallWithJson(
+		UObject* TestObject,
+		const FString& FunctionName,
+		const FString& InputJson)
+	{
+		check(TestObject);
+		UFunction* Function = TestObject->GetClass()->FindFunctionByName(*FunctionName);
+		if (!Function)
+		{
+			return MakeError(FString(TEXT("Function not found")));
+		}
+
+		auto ToolCall = MakeShared<FObjectFunctionToolCall>(TestObject, Function);
+		TSharedPtr<Internal::FToolCallScriptExceptionHandler> ExceptionHandler =
+			MakeShared<Internal::FToolCallScriptExceptionHandler>();
+
+		TFuture<FJsonValueOrError> ResultFuture = ToolCall->Execute(
+			FObjectFunctionToolCall::FFunctionInputParamsJson(
+				TInPlaceType<FString>(), InputJson),
+			ExceptionHandler);
+
+		return ResultFuture.Get();
+	}
+}
diff --git a/Engine/Plugins/Experimental/ToolsetRegistry/Source/ToolsetRegistry/Private/Tests/ToolsetJsonTest.cpp b/Engine/Plugins/Experimental/ToolsetRegistry/Source/ToolsetRegistry/Private/Tests/ToolsetJsonTest.cpp
--- a/Engine/Plugins/Experimental/ToolsetRegistry/Source/ToolsetRegistry/Private/Tests/ToolsetJsonTest.cpp
+++ b/Engine/Plugins/Experimental/ToolsetRegistry/Source/ToolsetRegistry/Private/Tests/ToolsetJsonTest.cpp
@@ -2,9 +2,12 @@
 
 #include "ToolsetJsonTest.h"
 
+#include "AssetRegistry/AssetData.h"
+#include "AssetRegistry/AssetRegistryModule.h"
 #include "CoreMinimal.h"
 #include "Containers/Array.h"
 #include "Containers/ArrayView.h"
+#include "Engine/Blueprint.h"
 #include "JsonObjectConverter.h"
 #include "Misc/AutomationTest.h"
 
@@ -14,6 +17,7 @@
 #include "ToolsetRegistry/ToolCallAsyncResult.h"
 #include "ToolsetRegistry/ToolCallAsyncResultString.h"
 #include "Tests/ToolCallAsyncResultTest.h"
+#include "Tests/ToolCallTestHelpers.h"
 
 #if WITH_DEV_AUTOMATION_TESTS
 
@@ -283,6 +287,224 @@ void FToolsetJsonSpec::Define()
 		TestEqual("RefCheck", Test.TestObject, TObjectPtr<UToolsetJsonTestObject>());
 	});
 
+	It("Accepts null for optional parameter", [this]()
+	{
+		TObjectPtr<UToolsetJsonTestObject> TestObject = NewObject<UToolsetJsonTestObject>();
+		UE::ToolsetRegistry::FJsonValueOrError Result =
+			UE::ToolsetRegistry::TestHelpers::ExecuteToolCallWithJson(
+				TestObject, TEXT("TestOptionalObject"), TEXT(R"({"testObject": {"refPath": ""}})"));
+		TestFalse("Should not have error", Result.HasError());
+	});
+
+	It("Rejects null for non-optional object properties", [this]()
+	{
+		TObjectPtr<UToolsetJsonTestObject> TestObject = NewObject<UToolsetJsonTestObject>();
+		UE::ToolsetRegistry::FJsonValueOrError Result =
+			UE::ToolsetRegistry::TestHelpers::ExecuteToolCallWithJson(
+				TestObject, TEXT("TestObjectParam"),
+				TEXT(R"({"testObject": {"refPath": ""}})"));
+
+		if (!TestTrue("Should have error", Result.HasError()))
+		{
+			return;
+		}
+		TestTrue("Error message", Result.GetError().Contains(TEXT("None is not valid value for property")));
+	});
+
+	It("Rejects invalid object path in object properties", [this]()
+	{
+		TObjectPtr<UToolsetJsonTestObject> TestObject = NewObject<UToolsetJsonTestObject>();
+		UE::ToolsetRegistry::FJsonValueOrError Result =
+			UE::ToolsetRegistry::TestHelpers::ExecuteToolCallWithJson(
+				TestObject, TEXT("TestObjectParam"),
+				TEXT(R"({"testObject": {"refPath": "/Invalid/Path/DoesNotExist"}})"));
+		if (!TestTrue("Should have error", Result.HasError()))
+		{
+			return;
+		}
+		TestTrue("Error contains validation message",
+				 Result.GetError().Contains(TEXT("is not valid Object for property")));
+	});
+
+	It("Rejects wrong object type in object properties", [this]()
+	{
+		TObjectPtr<UToolsetJsonTestObject> TestObject = NewObject<UToolsetJsonTestObject>();
+		UE::ToolsetRegistry::FJsonValueOrError Result =
+			UE::ToolsetRegistry::TestHelpers::ExecuteToolCallWithJson(
+				TestObject, TEXT("TestSpecObjectParam"),
+				FString::Printf(
+					TEXT(R"({"testObject": {"refPath": "%s"}})"), *UObject::StaticClass()->GetPathName()));
+		if (!TestTrue("Should have error", Result.HasError()))
+		{
+			return;
+		}
+		TestTrue("Error contains validation message",
+				 Result.GetError().Contains(TEXT("is not valid ToolsetJsonTestObject for property")));
+	});
+
+	It("Rejects invalid path for class properties", [this]()
+	{
+		TObjectPtr<UToolsetJsonTestObject> TestObject = NewObject<UToolsetJsonTestObject>();
+		UE::ToolsetRegistry::FJsonValueOrError Result =
+			UE::ToolsetRegistry::TestHelpers::ExecuteToolCallWithJson(
+				TestObject, TEXT("TestClassParam"),
+				TEXT(R"({"testClass": {"refPath": "/Invalid/Class/Path"}})"));
+		if (!TestTrue("Should have error", Result.HasError()))
+		{
+			return;
+		}
+		TestTrue("Error contains validation message",
+				 Result.GetError().Contains(TEXT("is not valid Class for property")));
+	});
+
+	It("Rejects non-class object for class properties", [this]()
+	{
+		TObjectPtr<UToolsetJsonTestObject> TestObj = NewObject<UToolsetJsonTestObject>();
+		TObjectPtr<UToolsetJsonTestObject> TestObject = NewObject<UToolsetJsonTestObject>();
+		UE::ToolsetRegistry::FJsonValueOrError Result =
+			UE::ToolsetRegistry::TestHelpers::ExecuteToolCallWithJson(
+				TestObject, TEXT("TestClassParam"),
+				FString::Printf(TEXT(R"({"testClass": {"refPath": "%s"}})"), *TestObj->GetPathName()));
+		if (!TestTrue("Should have error", Result.HasError()))
+		{
+			return;
+		}
+		TestTrue("Error contains validation message",
+				 Result.GetError().Contains(TEXT("is not valid Class for property")));
+	});
+
+	It("Rejects class not matching metaclass constraint in properties", [this]()
+	{
+		FString WrongClassPath = UObject::StaticClass()->GetPathName();
+		TObjectPtr<UToolsetJsonTestObject> TestObject = NewObject<UToolsetJsonTestObject>();
+		UE::ToolsetRegistry::FJsonValueOrError Result =
+			UE::ToolsetRegistry::TestHelpers::ExecuteToolCallWithJson(
+				TestObject, TEXT("TestSubClassParam"),
+				FString::Printf(TEXT(R"({"testClass": {"refPath": "%s"}})"), *WrongClassPath));
+		if (!TestTrue("Should have error", Result.HasError()))
+		{
+			return;
+		}
+		TestTrue("Error contains validation message",
+				 Result.GetError().Contains(TEXT("is not valid ToolsetJsonTestObject for property")));
+	});
+
+	It("Rejects wrong type for soft object reference", [this]()
+	{
+		FString InvalidJson = FString::Printf(
+			TEXT(R"({"testObject": {"refPath": "%s"}})"),
+			*UObject::StaticClass()->GetPathName());
+
+		TObjectPtr<UToolsetJsonTestObject> TestObject = NewObject<UToolsetJsonTestObject>();
+		UE::ToolsetRegistry::FJsonValueOrError Result =
+			UE::ToolsetRegistry::TestHelpers::ExecuteToolCallWithJson(
+				TestObject, TEXT("TestSoftObjectParam"), InvalidJson);
+
+		if (!TestTrue("Should have error", Result.HasError())) return;
+		TestTrue("Error message",
+			Result.GetError().Contains(TEXT("is not valid ToolsetJsonTestObject")));
+	});
+
+	It("Allows valid soft object reference", [this]()
+	{
+		// Create a valid test object asset
+		UToolsetJsonTestObject* TestObj = NewObject<UToolsetJsonTestObject>();
+		FString ValidJson = FString::Printf(
+			TEXT(R"({"testObject": {"refPath": "%s"}})"),
+			*TestObj->GetPathName());
+
+		TObjectPtr<UToolsetJsonTestObject> TestObject = NewObject<UToolsetJsonTestObject>();
+		UE::ToolsetRegistry::FJsonValueOrError Result =
+			UE::ToolsetRegistry::TestHelpers::ExecuteToolCallWithJson(
+				TestObject, TEXT("TestSoftObjectParam"), ValidJson);
+
+		TestFalse("Should not have error", Result.HasError());
+	});
+
+	It("Accepts Blueprint asset path for FSoftClassPath parameter", [this]()
+	{
+		// Regression test: before the fix, FSoftClassPath validation called AssetData.GetClass()
+		// which returns UBlueprint for Blueprint assets. Since UBlueprint is not a child of UClass,
+		// it incorrectly raised a script error. The fix adds an explicit check for UBlueprint.
+
+		// Create a transient Blueprint and register it with the Asset Registry so that
+		// AssetData.IsValid() is true, which is required to exercise the validation branch.
+		UBlueprint* Blueprint = NewObject<UBlueprint>(
+			GetTransientPackage(), TEXT("TestBlueprintForSoftClassPathValidation"), RF_Transient);
+		Blueprint->AddToRoot();
+
+		FAssetRegistryModule& AssetRegistryModule =
+			FModuleManager::LoadModuleChecked<FAssetRegistryModule>("AssetRegistry");
+		IAssetRegistry& AssetRegistry = AssetRegistryModule.Get();
+		AssetRegistry.AssetCreated(Blueprint);
+
+		ON_SCOPE_EXIT
+		{
+			AssetRegistry.AssetDeleted(Blueprint);
+			Blueprint->RemoveFromRoot();
+		};
+
+		FString BlueprintPath = Blueprint->GetPathName();
+		FString Json = FString::Printf(
+			TEXT(R"({"testClassPath": {"refPath": "%s"}})"), *BlueprintPath);
+
+		TObjectPtr<UToolsetJsonTestObject> TestObject = NewObject<UToolsetJsonTestObject>();
+		UE::ToolsetRegistry::FJsonValueOrError Result =
+			UE::ToolsetRegistry::TestHelpers::ExecuteToolCallWithJson(
+				TestObject, TEXT("TestSoftClassPathParam"), Json);
+
+		TestFalse("Blueprint asset path should be accepted for FSoftClassPath", Result.HasError());
+	});
+
+	It("Rejects non-class for soft class reference", [this]()
+	{
+		// Try to pass a regular object instance where a class is expected
+		UToolsetJsonTestObject* TestObj = NewObject<UToolsetJsonTestObject>();
+		FString InvalidJson = FString::Printf(
+			TEXT(R"({"testClass": {"refPath": "%s"}})"),
+			*TestObj->GetPathName());
+
+		TObjectPtr<UToolsetJsonTestObject> TestObject = NewObject<UToolsetJsonTestObject>();
+		UE::ToolsetRegistry::FJsonValueOrError Result =
+			UE::ToolsetRegistry::TestHelpers::ExecuteToolCallWithJson(
+				TestObject, TEXT("TestSoftClassParam"), InvalidJson);
+		if (!TestTrue("Should have error", Result.HasError())) return;
+		TestTrue("Error message",
+			Result.GetError().Contains(TEXT("is not valid Class for property")));
+	});
+
+	It("Rejects wrong class type for soft class reference", [this]()
+	{
+		// Try to pass UObject class where UToolsetJsonTestObject-derived class is expected
+		FString InvalidJson = FString::Printf(
+			TEXT(R"({"testClass": {"refPath": "%s"}})"),
+			*UObject::StaticClass()->GetPathName());
+
+		TObjectPtr<UToolsetJsonTestObject> TestObject = NewObject<UToolsetJsonTestObject>();
+		UE::ToolsetRegistry::FJsonValueOrError Result =
+			UE::ToolsetRegistry::TestHelpers::ExecuteToolCallWithJson(
+				TestObject, TEXT("TestSoftClassParam"), InvalidJson);
+
+		if (!TestTrue("Should have error", Result.HasError())) return;
+		TestTrue("Error message",
+			Result.GetError().Contains(TEXT("is not valid ToolsetJsonTestObject")));
+	});
+
+	It("Allows valid soft class reference", [this]()
+	{
+		// Pass a valid class reference
+		FString ValidJson = FString::Printf(
+			TEXT(R"({"testClass": {"refPath": "%s"}})"),
+			*UToolsetJsonTestObject::StaticClass()->GetPathName());
+
+		TObjectPtr<UToolsetJsonTestObject> TestObject = NewObject<UToolsetJsonTestObject>();
+		UE::ToolsetRegistry::FJsonValueOrError Result =
+			UE::ToolsetRegistry::TestHelpers::ExecuteToolCallWithJson(
+				TestObject, TEXT("TestSoftClassParam"), ValidJson);
+
+		TestFalse("Should not have error", Result.HasError());
+	});
+
 	It("Converts SoftPath property to reference schema", [this]()
 	{
 		UStruct* Struct = FToolsetJsonTest::StaticStruct();
diff --git a/Engine/Plugins/Experimental/ToolsetRegistry/Source/ToolsetRegistry/Private/Tests/ToolsetJsonTest.h b/Engine/Plugins/Experimental/ToolsetRegistry/Source/ToolsetRegistry/Private/Tests/ToolsetJsonTest.h
--- a/Engine/Plugins/Experimental/ToolsetRegistry/Source/ToolsetRegistry/Private/Tests/ToolsetJsonTest.h
+++ b/Engine/Plugins/Experimental/ToolsetRegistry/Source/ToolsetRegistry/Private/Tests/ToolsetJsonTest.h
@@ -3,6 +3,7 @@
 #pragma once
 
 #include "CoreMinimal.h"
+#include "Kismet/KismetSystemLibrary.h"
 #include "Templates/SharedPointer.h"
 #include "Templates/SubclassOf.h"
 
@@ -39,6 +40,51 @@ public:
 	void TestDefault(
 		UObject* TestObject = nullptr, FColor TestColor = FColor(10, 20, 30, 40),
 		FLinearColor TestLinearColor = FLinearColor(0.1, 0.2, 0.3, 0.4)) {}
+
+	UFUNCTION(BlueprintCallable, Category = "ToolsetTest")
+	void TestOptionalObject(
+		UObject* TestObject = nullptr) {}
+
+	UFUNCTION(BlueprintCallable, Category = "ToolsetTest")
+	void TestObjectParam(
+		UObject* TestObject) {}
+
+	UFUNCTION(BlueprintCallable, Category = "ToolsetTest")
+	void TestSpecObjectParam(
+		UToolsetJsonTestObject* TestObject) {}
+
+	UFUNCTION(BlueprintCallable, Category = "ToolsetTest")
+	void TestClassParam(
+		UClass* TestClass) {}
+
+	UFUNCTION(BlueprintCallable, Category = "ToolsetTest")
+	void TestSubClassParam(
+		TSubclassOf<UToolsetJsonTestObject> TestClass) {}
+
+	UFUNCTION(BlueprintCallable, Category = "ToolsetTest")
+	void TestSoftObjectParam(
+		TSoftObjectPtr<UToolsetJsonTestObject> TestObject) {}
+
+	UFUNCTION(BlueprintCallable, Category = "ToolsetTest")
+	void TestSoftClassParam(
+		TSoftClassPtr<UToolsetJsonTestObject> TestClass) {}
+
+	UFUNCTION(BlueprintCallable, Category = "ToolsetTest")
+	void TestSoftClassPathParam(
+		FSoftClassPath TestClassPath) {}
+
+	// Functions for testing FToolCallScriptExceptionHandler
+	UFUNCTION(BlueprintCallable, Category = "ToolsetTest")
+	void TestMultipleRefs(
+		UToolsetJsonTestObject* FirstObject,
+		UToolsetJsonTestObject* SecondObject) {}
+
+	UFUNCTION(BlueprintCallable, Category = "ToolsetTest")
+	void TestEmptyError()
+	{
+		// Raise an explicitly empty error
+		UKismetSystemLibrary::RaiseScriptError(FString());
+	}
 };
 
 USTRUCT(BlueprintType)
diff --git a/Engine/Plugins/Experimental/ToolsetRegistry/Source/ToolsetRegistry/Private/ToolsetRegistry/BlueprintLibraryToolset.cpp b/Engine/Plugins/Experimental/ToolsetRegistry/Source/ToolsetRegistry/Private/ToolsetRegistry/BlueprintLibraryToolset.cpp
--- a/Engine/Plugins/Experimental/ToolsetRegistry/Source/ToolsetRegistry/Private/ToolsetRegistry/BlueprintLibraryToolset.cpp
+++ b/Engine/Plugins/Experimental/ToolsetRegistry/Source/ToolsetRegistry/Private/ToolsetRegistry/BlueprintLibraryToolset.cpp
@@ -3,18 +3,9 @@
 #include "ToolsetRegistry/BlueprintLibraryToolset.h"
 
 #include "Async/Future.h"
-#include "Blueprint/BlueprintExceptionInfo.h"
 #include "Containers/UnrealString.h"
-#include "Dom/JsonObject.h"
-#include "Kismet/BlueprintFunctionLibrary.h"
-#include "Misc/CoreMisc.h"
-#include "Misc/Optional.h"
-#include "Policies/CondensedJsonPrintPolicy.h"
-#include "Serialization/JsonWriter.h"
-#include "Serialization/JsonSerializer.h"
-#include "Templates/UnrealTemplate.h"
+#include "UObject/FieldIterator.h"
 
-#include "ToolsetRegistry/DelegateHandle.h"
 #include "ToolsetRegistry/JsonConversion.h"
 #include "ToolsetRegistry/JsonValueOrError.h"
 #include "ToolsetRegistry/Module.h"
@@ -26,140 +17,6 @@ namespace UE::ToolsetRegistry
 {
 	namespace
 	{
-		// Register for tracking of script exceptions for the lifetime of this object.
-		class FToolCallScriptExceptionHandler
-		{
-		  public:
-			FToolCallScriptExceptionHandler()
-				: ScriptExceptionHandle(
-					FDelegateHandleRaii::Create(
-						FBlueprintCoreDelegates::OnScriptException,
-						FBlueprintCoreDelegates::OnScriptException.AddRaw(
-							this, &FToolCallScriptExceptionHandler::OnScriptException))),
-				  RuntimeScriptExceptionHandle(
-					[this](ELogVerbosity::Type Verbosity, const TCHAR* ExceptionMessage, const TCHAR* StackMessage)
-					{
-						AddToExceptionMessage(FString(ExceptionMessage));
-
-						// Also log to console to make more obvious to tool developers.
-						UE_LOG(LogToolsetRegistry, Warning, TEXT("Error running tool:\n%s"), ExceptionMessage);
-					})
-			{}
-
-			~FToolCallScriptExceptionHandler() = default;
-
-			// Track script exceptions raised by UKismetSystemLibrary::RaiseScriptError or
-			// FBlueprintCoreDelegates::ThrowScriptException via FBlueprintCoreDelegates::OnScriptException
-			// and copy into local storage.
-			void OnScriptException(
-				const UObject* ActiveObject, const FFrame& StackFrame, const FBlueprintExceptionInfo& Info)
-			{
-				TOptional<FString> ErrorMessage = GetBlueprintExceptionMessage(Info, StackFrame);
-				if (ErrorMessage.IsSet())
-				{
-					AddToExceptionMessage(ErrorMessage.GetValue());
-				}
-			}
-
-			FString GetException() const
-			{
-				return ScriptException;
-			}
-
-		  private:
-			FString GetBackupExceptionMessage(const FString& ExceptionType, const FFrame& StackFrame)
-			{
-				// Generate as useful a message as we can if the exception didn't provide one.
-				FString FallbackExceptionMessage;
-				if (StackFrame.PreviousTrackingFrame &&
-					StackFrame.PreviousTrackingFrame->Object &&
-					StackFrame.PreviousTrackingFrame->Node)
-				{
-					FallbackExceptionMessage = FString::Printf(
-						TEXT("%s in %s.%s"), *ExceptionType,
-						*StackFrame.PreviousTrackingFrame->Object->GetName(),
-						*StackFrame.PreviousTrackingFrame->Node->GetName());
-				}
-				else
-				{
-					FallbackExceptionMessage = ExceptionType;
-				}
-				return FallbackExceptionMessage;
-			}
-
-			TOptional<FString> GetBlueprintExceptionMessage(const FBlueprintExceptionInfo& Info, const FFrame& StackFrame)
-			{
-				TOptional<FString> ExceptionMessage;
-				switch (Info.GetType())
-				{
-					// Ignore breakpoints & tracepoints.
-					case EBlueprintExceptionType::Breakpoint:
-						// Ignore and fall through
-					case EBlueprintExceptionType::Tracepoint:
-						// Ignore and fall through
-					case EBlueprintExceptionType::WireTracepoint:
-						// Return unset message, to be ignored.
-						break;
-					case EBlueprintExceptionType::AccessViolation:
-						ExceptionMessage = Info.GetDescription().IsEmpty() ?
-							GetBackupExceptionMessage(TEXT("AccessViolation"), StackFrame) :
-							Info.GetDescription().ToString();
-						break;
-					case EBlueprintExceptionType::InfiniteLoop:
-						ExceptionMessage = Info.GetDescription().IsEmpty() ?
-							GetBackupExceptionMessage(TEXT("InfiniteLoop"), StackFrame) :
-							Info.GetDescription().ToString();
-						break;
-					case EBlueprintExceptionType::NonFatalError:
-						ExceptionMessage = Info.GetDescription().IsEmpty() ?
-							GetBackupExceptionMessage(TEXT("NonFatalError"), StackFrame) :
-							Info.GetDescription().ToString();
-						break;
-					case EBlueprintExceptionType::FatalError:
-						ExceptionMessage = Info.GetDescription().IsEmpty() ?
-							GetBackupExceptionMessage(TEXT("FatalError"), StackFrame) :
-							Info.GetDescription().ToString();
-						break;
-					case EBlueprintExceptionType::AbortExecution:
-						ExceptionMessage = Info.GetDescription().IsEmpty() ?
-							GetBackupExceptionMessage(TEXT("AbortExecution"), StackFrame) :
-							Info.GetDescription().ToString();
-						break;
-					case EBlueprintExceptionType::UserRaisedError:
-						ExceptionMessage = Info.GetDescription().IsEmpty() ?
-							GetBackupExceptionMessage(TEXT("UserRaisedError"), StackFrame) :
-							Info.GetDescription().ToString();
-						break;
-					default:
-						checkNoEntry();
-						break;
-				}
-				return ExceptionMessage;
-			}
-
-			void AddToExceptionMessage(const FString& NewMessage)
-			{
-				if (NewMessage.IsEmpty())
-				{
-					return;
-				}
-				else if (ScriptException.IsEmpty())
-				{
-					ScriptException = NewMessage;
-				}
-				else
-				{
-					ScriptException += FString::Printf(TEXT("\n%s"), *NewMessage);
-				}
-			}
-
-			// Handler for FBlueprintCoreDelegates::OnScriptException
-			FDelegateHandleRaii ScriptExceptionHandle;
-			// Handler for FScriptExceptionHandler
-			FScopedScriptExceptionHandler RuntimeScriptExceptionHandle;
-			// Raised script exception message.
-			FString ScriptException;
-		};
 	}
 
 	FBlueprintLibraryToolset::FBlueprintLibraryToolset(UClass* LibraryClass) :
@@ -263,7 +120,7 @@ namespace UE::ToolsetRegistry
 		ensureAlways(IsInGameThread());
 
 		FString ToolNameCopy = ToolName;
-		return ToolCall->Execute(JsonParams).Next(
+		return ToolCall->Execute(JsonParams, ExceptionHandler).Next(
 			[
 				Toolset,  // Capture to keep alive.
 				ToolCall = ToolCall,
diff --git a/Engine/Plugins/Experimental/ToolsetRegistry/Source/ToolsetRegistry/Private/ToolsetRegistry/ObjectFunctionToolCall.cpp b/Engine/Plugins/Experimental/ToolsetRegistry/Source/ToolsetRegistry/Private/ToolsetRegistry/ObjectFunctionToolCall.cpp
--- a/Engine/Plugins/Experimental/ToolsetRegistry/Source/ToolsetRegistry/Private/ToolsetRegistry/ObjectFunctionToolCall.cpp
+++ b/Engine/Plugins/Experimental/ToolsetRegistry/Source/ToolsetRegistry/Private/ToolsetRegistry/ObjectFunctionToolCall.cpp
@@ -282,7 +282,8 @@ namespace UE::ToolsetRegistry
 	}
 
 	TFuture<FJsonValueOrError> FObjectFunctionToolCall::Execute(
-		const TOptional<FFunctionInputParamsJson>& FunctionInputParamsJson) const
+		const TOptional<FFunctionInputParamsJson>& FunctionInputParamsJson,
+		TSharedPtr<FToolCallScriptExceptionHandler> ExceptionHandler) const
 	{
 		// Get function input params JSON as a JSON object.
 		TSharedPtr<FJsonObject> FunctionInputParamsJsonObject;
@@ -327,6 +328,15 @@ namespace UE::ToolsetRegistry
 			}
 		};
 
+		// TODO: A lot of this logic is duplicated from ScriptCore and should be ported to use
+		// existing function invocation logic.
+
+		// StackFrame must be set up before calling JsonDataToStruct so that errors raised via
+		// RaiseScriptError during input validation can be reliably detected for C++ tool calls (Python tool
+		// calls get a stack frame set up by the scripting runtime).
+		FFrame StackFrame(InstanceObject, Function, FrameMemory, nullptr,
+			Function->ChildProperties);
+
 		// Copy the function input params memory into the function's frame memory.
 		if (FrameMemory &&
 			!ToolsetJson::JsonDataToStruct(
@@ -344,16 +354,19 @@ namespace UE::ToolsetRegistry
 				*JsonToString(FunctionInputParamsJsonObject.ToSharedRef())));
 		}
 
+		// If there were errors raised in converting JSON params, return with error without invoking the
+		// function.
+		if (ExceptionHandler && !ExceptionHandler->GetException().IsEmpty())
+		{
+			FString ErrorMessage = ExceptionHandler->GetException();
+			return FJsonValueOrErrorFuture::MakeError(
+					FString::Printf(TEXT("Parameter error: %s."), *ErrorMessage));
+		}
+
 		// In case this is a Blueprint function that has a world context object parameter, 
 		// give that parameter (in the frame memory) the World.
 		MaybeSetWorldContextPropertyInFrameMemory(Function, FrameMemory);
 
-		// Invoke function.
-		// TODO: A lot of this logic is duplicated from ScriptCore and should be ported to use 
-		// existing function invocation logic.
-		FFrame StackFrame(InstanceObject, Function, FrameMemory, nullptr,
-			Function->ChildProperties);
-
 		// If the function has out parameters, fill the stack frame's out parameter info with the 
 		// info for those params.
 		if (Function->HasAnyFunctionFlags(FUNC_HasOutParms))
diff --git a/Engine/Plugins/Experimental/ToolsetRegistry/Source/ToolsetRegistry/Private/ToolsetRegistry/ObjectFunctionToolCall.h b/Engine/Plugins/Experimental/ToolsetRegistry/Source/ToolsetRegistry/Private/ToolsetRegistry/ObjectFunctionToolCall.h
--- a/Engine/Plugins/Experimental/ToolsetRegistry/Source/ToolsetRegistry/Private/ToolsetRegistry/ObjectFunctionToolCall.h
+++ b/Engine/Plugins/Experimental/ToolsetRegistry/Source/ToolsetRegistry/Private/ToolsetRegistry/ObjectFunctionToolCall.h
@@ -7,10 +7,10 @@
 #include "Misc/NotNull.h"
 #include "Misc/TVariant.h"
 #include "Templates/SharedPointer.h"
-#include "Templates/ValueOrError.h"
 #include "UObject/Class.h"
 
 #include "ToolsetRegistry/JsonValueOrError.h"
+#include "ToolsetRegistry/ToolCallScriptExceptionHandler.h"
 
 //
 // UE::ToolsetRegistry::FObjectFunctionToolCall
@@ -66,7 +66,8 @@ namespace UE::ToolsetRegistry
 		 */
 		TFuture<FJsonValueOrError> Execute(
 			const TOptional<FFunctionInputParamsJson>& FunctionInputParamsJson = 
-				TOptional<FFunctionInputParamsJson>()) const;
+			TOptional<FFunctionInputParamsJson>(),
+			TSharedPtr<Internal::FToolCallScriptExceptionHandler> ExceptionHandler = nullptr) const;
 
 	private:
 		FJsonValueOrError BuildValidFunctionInputParamsJsonObject(
diff --git a/Engine/Plugins/Experimental/ToolsetRegistry/Source/ToolsetRegistry/Private/ToolsetRegistry/ReferenceConverter.cpp b/Engine/Plugins/Experimental/ToolsetRegistry/Source/ToolsetRegistry/Private/ToolsetRegistry/ReferenceConverter.cpp
--- a/Engine/Plugins/Experimental/ToolsetRegistry/Source/ToolsetRegistry/Private/ToolsetRegistry/ReferenceConverter.cpp
+++ b/Engine/Plugins/Experimental/ToolsetRegistry/Source/ToolsetRegistry/Private/ToolsetRegistry/ReferenceConverter.cpp
@@ -2,11 +2,15 @@
 
 #include "ReferenceConverter.h"
 
+#include "AssetRegistry/AssetData.h"
+#include "AssetRegistry/AssetRegistryModule.h"
+#include "Engine/Blueprint.h"
 #include "JsonObjectConverter.h"
 #include "JsonSchema/JsonSchemaGenerator.h"
 #include "Kismet/KismetSystemLibrary.h"
 #include "Misc/StringOutputDevice.h"
 
+#include "ToolsetRegistry/Module.h"
 #include "ToolsetRegistry/PropertyAccessors.h"
 #include "ToolsetRegistry/JsonConversion.h"
 
@@ -16,6 +20,52 @@ namespace
 	{
 		return String.IsEmpty() || String.Equals("None", ESearchCase::IgnoreCase) || String.Equals("null", ESearchCase::IgnoreCase);
 	}
+
+	// Helper to validate soft reference type without loading the asset
+	bool ValidateSoftReferenceType(
+		const FString& RefPath,
+		const UClass* ExpectedClass,
+		FString& OutErrorMessage)
+	{
+		// Try to get asset data from registry
+		FAssetRegistryModule& AssetRegistryModule =
+			FModuleManager::LoadModuleChecked<FAssetRegistryModule>("AssetRegistry");
+		IAssetRegistry& AssetRegistry = AssetRegistryModule.Get();
+
+		FSoftObjectPath SoftPath(RefPath);
+		FAssetData AssetData = AssetRegistry.GetAssetByObjectPath(SoftPath);
+
+		if (!AssetData.IsValid())
+		{
+			// Asset not in registry - could be:
+			// 1. Invalid path (asset doesn't exist)
+			// 2. Not scanned yet
+			// 3. Generated at runtime
+			// We'll allow it and let it fail at load time if invalid
+			return true; // Permissive: don't block unknown assets
+		}
+
+		// Get the asset's class
+		UClass* AssetClass = AssetData.GetClass();
+		if (!AssetClass)
+		{
+			// Can't determine class from registry
+			return true; // Permissive
+		}
+
+		// Check if asset class is compatible with expected class
+		if (!AssetClass->IsChildOf(ExpectedClass))
+		{
+			OutErrorMessage = FString::Printf(
+				TEXT("%s is not valid %s (asset is %s) for property"),
+				*RefPath,
+				*ExpectedClass->GetName(),
+				*AssetClass->GetName());
+			return false;
+		}
+
+		return true;
+	}
 }
 
 namespace UE::ToolsetRegistry
@@ -165,19 +215,170 @@ namespace UE::ToolsetRegistry
 			RefPath = JsonValue->AsString();
 		}
 
+		bool IsNull = false;
 		if (JsonValue->Type == EJson::Null || StringMeansNull(RefPath))
 		{
 			RefPath = "None";
+			IsNull = true;
+		}
+
+		check(!RefPath.IsEmpty());
+
+		if (IsNull)
+		{
+			// Raise error if null property is a function param and is not either 1) explicitly a
+			// FOptionalProperty, or 2) a reference with a default value of null.
+
+			// For FOptionalProperty, check both the property and its owner: Python optional params work by
+			// generating a transient value property for the value owned by the optional property.
+			const FOptionalProperty* OptionalProperty = CastField<const FOptionalProperty>(Property);
+			if (!OptionalProperty)
+			{
+				OptionalProperty = CastField<const FOptionalProperty>(Property->Owner.ToField());
+			}
+
+			if (!OptionalProperty)
+			{
+				if (const UFunction* OwningFunction = Cast<const UFunction>(Property->Owner.ToUObject()))
+				{
+					bool HasNullDefault = false;
+					FString MetadataKey = FString::Printf(TEXT("CPP_Default_%s"), *Property->GetName());
+					if (OwningFunction->HasMetaData(*MetadataKey))
+					{
+						FString DefaultValue = OwningFunction->GetMetaData(*MetadataKey);
+						HasNullDefault = StringMeansNull(DefaultValue);  // Reuse existing helper
+					}
+
+					if (!HasNullDefault)
+					{
+						UKismetSystemLibrary::RaiseScriptError(
+							FString::Printf(TEXT("%s is not valid value for property '%s'"),
+											*RefPath, *Property->GetName()));
+						return true;
+					}
+				}
+			}
+		}
+		else
+		{
+			FObjectProperty* ObjectProperty = CastField<FObjectProperty>(Property);
+			const FClassProperty* ClassProperty = CastField<const FClassProperty>(Property);
+			const FSoftClassProperty* SoftClassProperty = CastField<const FSoftClassProperty>(Property);
+
+			if (ClassProperty || SoftClassProperty ||
+				(ObjectProperty && ObjectProperty->PropertyClass == UClass::StaticClass()))
+			{
+				const UClass* MetaClass = ClassProperty ? ClassProperty->MetaClass :
+										  (SoftClassProperty ? SoftClassProperty->MetaClass : nullptr);
+				UObject* Object = FSoftObjectPath(RefPath).ResolveObject();
+				if (!Object)
+				{
+					Object = FSoftObjectPath(RefPath).TryLoad();
+				}
+				if (!Object || !Object->IsA(UClass::StaticClass()))
+				{
+					UKismetSystemLibrary::RaiseScriptError(
+						FString::Printf(TEXT("%s is not valid %s for property '%s'"),
+										*RefPath, *UClass::StaticClass()->GetName(), *Property->GetName()));
+					return true;
+				}
+				else if (MetaClass && !Cast<UClass>(Object)->IsChildOf(MetaClass))
+				{
+					UKismetSystemLibrary::RaiseScriptError(
+						FString::Printf(TEXT("%s is not valid %s for property '%s'"),
+										*RefPath, *MetaClass->GetName(), *Property->GetName()));
+					return true;
+				}
+			}
+			else if (ObjectProperty)
+			{
+				UObject* Object = FSoftObjectPath(RefPath).ResolveObject();
+				if (!Object)
+				{
+					Object = FSoftObjectPath(RefPath).TryLoad();
+				}
+				if (!Object || !Object->IsA(ObjectProperty->PropertyClass))
+				{
+					UKismetSystemLibrary::RaiseScriptError(
+						FString::Printf(
+							TEXT("%s is not valid %s for property '%s'"),
+							*RefPath, *ObjectProperty->PropertyClass->GetName(), *Property->GetName()));
+					return true;
+				}
+			}
+			// Validate soft references using Asset Registry
+			else if (FSoftObjectProperty* SoftObjectProperty = CastField<FSoftObjectProperty>(Property))
+			{
+				FString ErrorMessage;
+				if (!ValidateSoftReferenceType(RefPath, SoftObjectProperty->PropertyClass, ErrorMessage))
+				{
+					UKismetSystemLibrary::RaiseScriptError(
+						FString::Printf(TEXT("%s '%s'"), *ErrorMessage, *Property->GetName()));
+					return true;
+				}
+			}
+			else if (FWeakObjectProperty* WeakObjectProperty = CastField<FWeakObjectProperty>(Property))
+			{
+				// WeakObjectProperty: validate using Asset Registry (permissive for runtime objects)
+				// Weak references typically point to objects that are already loaded
+				FString ErrorMessage;
+				if (!ValidateSoftReferenceType(RefPath, WeakObjectProperty->PropertyClass, ErrorMessage))
+				{
+					UKismetSystemLibrary::RaiseScriptError(
+						FString::Printf(TEXT("%s '%s'"), *ErrorMessage, *Property->GetName()));
+					return true;
+				}
+			}
+			else if (const FStructProperty* StructProperty = CastField<FStructProperty>(Property))
+			{
+				const UScriptStruct* ScriptStruct = StructProperty->Struct;
+
+				if (ScriptStruct == FSoftObjectPath::StaticStruct())
+				{
+					// FSoftObjectPath: generic soft reference to any UObject
+					// No specific class constraint, so we can only check if asset exists
+					FAssetRegistryModule& AssetRegistryModule =
+						FModuleManager::LoadModuleChecked<FAssetRegistryModule>("AssetRegistry");
+					FAssetData AssetData = AssetRegistryModule.Get().GetAssetByObjectPath(FSoftObjectPath(RefPath));
+
+					// Permissive: only log warning if not found, don't error
+					if (!AssetData.IsValid())
+					{
+						UE_LOG(LogToolsetRegistry, Warning,
+							   TEXT("Soft reference path '%s' not found in Asset Registry for property '%s'"),
+							   *RefPath, *Property->GetName());
+					}
+				}
+				else if (ScriptStruct == FSoftClassPath::StaticStruct())
+				{
+					// FSoftClassPath: validate it's actually a class
+					FAssetRegistryModule& AssetRegistryModule =
+						FModuleManager::LoadModuleChecked<FAssetRegistryModule>("AssetRegistry");
+					FAssetData AssetData = AssetRegistryModule.Get().GetAssetByObjectPath(FSoftObjectPath(RefPath));
+
+					if (AssetData.IsValid())
+					{
+						UClass* AssetClass = AssetData.GetClass();
+						if (AssetClass &&
+							!AssetClass->IsChildOf(UClass::StaticClass()) &&
+							!AssetClass->IsChildOf(UBlueprint::StaticClass()))
+						{
+							UKismetSystemLibrary::RaiseScriptError(
+								FString::Printf(TEXT("%s is not a valid class for property '%s'"),
+												*RefPath, *Property->GetName()));
+							return true;
+						}
+					}
+				}
+			}
 		}
 
 		FStringOutputDevice ImportError;
 		Property->ImportText_Direct(*RefPath, OutValue, nullptr, 0, &ImportError);
 		if (!ImportError.IsEmpty())
 		{
-			UKismetSystemLibrary::RaiseScriptError(FString::Printf(
-				TEXT("%s is not valid JSON for Property %s"),
-				*Internal::JsonToString(JsonValue.ToSharedRef()), *Property->GetName()));
+			UKismetSystemLibrary::RaiseScriptError(ImportError);
 		}
 		return true;
 	}
-}
\ No newline at end of file
+}
diff --git a/Engine/Plugins/Experimental/ToolsetRegistry/Source/ToolsetRegistry/Private/ToolsetRegistry/ToolCallScriptExceptionHandler.cpp b/Engine/Plugins/Experimental/ToolsetRegistry/Source/ToolsetRegistry/Private/ToolsetRegistry/ToolCallScriptExceptionHandler.cpp
new file mode 100644
--- /dev/null
+++ b/Engine/Plugins/Experimental/ToolsetRegistry/Source/ToolsetRegistry/Private/ToolsetRegistry/ToolCallScriptExceptionHandler.cpp
@@ -0,0 +1,135 @@
+// Copyright Epic Games, Inc. All Rights Reserved.
+
+#include "ToolsetRegistry/ToolCallScriptExceptionHandler.h"
+
+#include "Blueprint/BlueprintExceptionInfo.h"
+#include "UObject/Stack.h"
+
+#include "ToolsetRegistry/DelegateHandle.h"
+#include "ToolsetRegistry/Module.h"
+
+namespace UE::ToolsetRegistry::Internal
+{
+
+	FToolCallScriptExceptionHandler::FToolCallScriptExceptionHandler()
+		: ScriptExceptionHandle(
+			FDelegateHandleRaii::Create(
+				FBlueprintCoreDelegates::OnScriptException,
+				FBlueprintCoreDelegates::OnScriptException.AddRaw(
+					this, &FToolCallScriptExceptionHandler::OnScriptException))),
+		  RuntimeScriptExceptionHandle(
+			  [this](ELogVerbosity::Type Verbosity, const TCHAR* ExceptionMessage, const TCHAR* StackMessage)
+			  {
+				  AddToExceptionMessage(FString(ExceptionMessage));
+
+				  // Also log to console to make more obvious to tool developers.
+				  UE_LOG(LogToolsetRegistry, Warning, TEXT("Error running tool:\n%s"), ExceptionMessage);
+			  })
+	{}
+
+	void FToolCallScriptExceptionHandler::OnScriptException(
+			const UObject* ActiveObject, const FFrame& StackFrame, const FBlueprintExceptionInfo& Info)
+	{
+		TOptional<FString> ErrorMessage = GetBlueprintExceptionMessage(Info, StackFrame);
+		if (ErrorMessage.IsSet())
+		{
+			AddToExceptionMessage(ErrorMessage.GetValue());
+		}
+	}
+
+	FString FToolCallScriptExceptionHandler::GetBackupExceptionMessage(
+		const FString& ExceptionType, const FFrame& StackFrame)
+	{
+		// Generate as useful a message as we can if the exception didn't provide one.
+		FString FallbackExceptionMessage;
+		if (StackFrame.PreviousTrackingFrame &&
+			StackFrame.PreviousTrackingFrame->Object &&
+			StackFrame.PreviousTrackingFrame->Node)
+		{
+			FallbackExceptionMessage = FString::Printf(
+				TEXT("%s in %s.%s"), *ExceptionType,
+				*StackFrame.PreviousTrackingFrame->Object->GetName(),
+				*StackFrame.PreviousTrackingFrame->Node->GetName());
+		}
+		else if (StackFrame.Object && StackFrame.Node)
+		{
+			FallbackExceptionMessage = FString::Printf(
+				TEXT("%s in %s.%s"), *ExceptionType,
+				*StackFrame.Object->GetName(),
+				*StackFrame.Node->GetName());
+		}
+		else
+		{
+			FallbackExceptionMessage = ExceptionType;
+		}
+		return FallbackExceptionMessage;
+	}
+
+	TOptional<FString> FToolCallScriptExceptionHandler::GetBlueprintExceptionMessage(
+		const FBlueprintExceptionInfo& Info, const FFrame& StackFrame)
+	{
+		TOptional<FString> ExceptionMessage;
+		switch (Info.GetType())
+		{
+			// Ignore breakpoints & tracepoints.
+			case EBlueprintExceptionType::Breakpoint:
+				// Ignore and fall through
+			case EBlueprintExceptionType::Tracepoint:
+				// Ignore and fall through
+			case EBlueprintExceptionType::WireTracepoint:
+				// Return unset message, to be ignored.
+				break;
+			case EBlueprintExceptionType::AccessViolation:
+				ExceptionMessage = Info.GetDescription().IsEmpty() ?
+					GetBackupExceptionMessage(TEXT("AccessViolation"), StackFrame) :
+					Info.GetDescription().ToString();
+				break;
+			case EBlueprintExceptionType::InfiniteLoop:
+				ExceptionMessage = Info.GetDescription().IsEmpty() ?
+					GetBackupExceptionMessage(TEXT("InfiniteLoop"), StackFrame) :
+					Info.GetDescription().ToString();
+				break;
+			case EBlueprintExceptionType::NonFatalError:
+				ExceptionMessage = Info.GetDescription().IsEmpty() ?
+					GetBackupExceptionMessage(TEXT("NonFatalError"), StackFrame) :
+					Info.GetDescription().ToString();
+				break;
+			case EBlueprintExceptionType::FatalError:
+				ExceptionMessage = Info.GetDescription().IsEmpty() ?
+					GetBackupExceptionMessage(TEXT("FatalError"), StackFrame) :
+					Info.GetDescription().ToString();
+				break;
+			case EBlueprintExceptionType::AbortExecution:
+				ExceptionMessage = Info.GetDescription().IsEmpty() ?
+					GetBackupExceptionMessage(TEXT("AbortExecution"), StackFrame) :
+					Info.GetDescription().ToString();
+				break;
+			case EBlueprintExceptionType::UserRaisedError:
+				ExceptionMessage = Info.GetDescription().IsEmpty() ?
+					GetBackupExceptionMessage(TEXT("UserRaisedError"), StackFrame) :
+					Info.GetDescription().ToString();
+				break;
+			default:
+				checkNoEntry();
+				break;
+		}
+		return ExceptionMessage;
+	}
+
+	void FToolCallScriptExceptionHandler::AddToExceptionMessage(const FString& NewMessage)
+	{
+		if (NewMessage.IsEmpty())
+		{
+			return;
+		}
+		else if (ScriptException.IsEmpty())
+		{
+			ScriptException = NewMessage;
+		}
+		else
+		{
+			ScriptException += FString::Printf(TEXT("\n%s"), *NewMessage);
+		}
+	}
+}
+
diff --git a/Engine/Plugins/Experimental/ToolsetRegistry/Source/ToolsetRegistry/Private/ToolsetRegistry/ToolCallScriptExceptionHandler.h b/Engine/Plugins/Experimental/ToolsetRegistry/Source/ToolsetRegistry/Private/ToolsetRegistry/ToolCallScriptExceptionHandler.h
new file mode 100644
--- /dev/null
+++ b/Engine/Plugins/Experimental/ToolsetRegistry/Source/ToolsetRegistry/Private/ToolsetRegistry/ToolCallScriptExceptionHandler.h
@@ -0,0 +1,45 @@
+// Copyright Epic Games, Inc. All Rights Reserved.
+
+#pragma once
+
+#include "Blueprint/BlueprintExceptionInfo.h"
+#include "UObject/Stack.h"
+
+#include "ToolsetRegistry/DelegateHandle.h"
+
+namespace UE::ToolsetRegistry::Internal
+{
+	// Register for tracking of script exceptions for the lifetime of this object.
+	class FToolCallScriptExceptionHandler
+	{
+	  public:
+		FToolCallScriptExceptionHandler();
+
+		~FToolCallScriptExceptionHandler() = default;
+
+		// Track script exceptions raised by UKismetSystemLibrary::RaiseScriptError or
+		// FBlueprintCoreDelegates::ThrowScriptException via FBlueprintCoreDelegates::OnScriptException
+		// and copy into local storage.
+		void OnScriptException(
+			const UObject* ActiveObject, const FFrame& StackFrame, const FBlueprintExceptionInfo& Info);
+
+		FString GetException() const
+		{
+			return ScriptException;
+		}
+
+	  private:
+		FString GetBackupExceptionMessage(const FString& ExceptionType, const FFrame& StackFrame);
+
+		TOptional<FString> GetBlueprintExceptionMessage(const FBlueprintExceptionInfo& Info, const FFrame& StackFrame);
+
+		void AddToExceptionMessage(const FString& NewMessage);
+
+		// Handler for FBlueprintCoreDelegates::OnScriptException
+		FDelegateHandleRaii ScriptExceptionHandle;
+		// Handler for FScriptExceptionHandler
+		FScopedScriptExceptionHandler RuntimeScriptExceptionHandle;
+		// Raised script exception message.
+		FString ScriptException;
+	};
+}
