Fix hanging EditorAppToolset tests by fixing the interplay between latent tests and commands.

#Synced-CL 51698457

diff --git a/Engine/Plugins/Experimental/ToolsetRegistry/Source/ToolsetRegistry/Private/Tests/EditorAppToolsetTest.cpp b/Engine/Plugins/Experimental/ToolsetRegistry/Source/ToolsetRegistry/Private/Tests/EditorAppToolsetTest.cpp
--- a/Engine/Plugins/Experimental/ToolsetRegistry/Source/ToolsetRegistry/Private/Tests/EditorAppToolsetTest.cpp
+++ b/Engine/Plugins/Experimental/ToolsetRegistry/Source/ToolsetRegistry/Private/Tests/EditorAppToolsetTest.cpp
@@ -156,29 +156,26 @@ void FEditorAppToolsetSpec::Define()
 		TestFalse("Error is set", Result->Error.IsEmpty());
 	});
 
-	LatentIt(TEXT("CaptureAssetImage renders a thumbnail for a static mesh"),
-		[this](const FDoneDelegate& Done)
+	It(TEXT("CaptureAssetImage renders a thumbnail for a static mesh"),
+		[this]()
 		{
 			if (!FApp::CanEverRender())
 			{
 				// Rendering is disabled; skipping thumbnail test.
-				Done.Execute();
 				return;
 			}
 
 			UToolCallAsyncResultImage* Result = UEditorAppToolset::CaptureAssetImage(
 				TEXT("/Engine/BasicShapes/Cube.Cube"));
-			TestNotNull("Result", Result);
-			if (!Result)
+			if (!TestNotNull("Result", Result))
 			{
-				Done.Execute();
 				return;
 			}
 
 			TStrongObjectPtr<UToolCallAsyncResultImage> StrongResult(Result);
 
 			AddCommand(new FFunctionLatentCommand(
-				[this, Done, StrongResult]() mutable -> bool
+				[this, StrongResult]() mutable -> bool
 				{
 					if (!StrongResult->bIsComplete)
 					{
@@ -188,18 +185,16 @@ void FEditorAppToolsetSpec::Define()
 					TestFalse("Has image data", StrongResult->Value.Data.IsEmpty());
 					TestEqual("MIME type",
 						StrongResult->Value.MimeType, FString(TEXT("image/png")));
-					Done.Execute();
 					return true; // done
 				}));
 		});
 
-	LatentIt(TEXT("CaptureAssetImage screenshots the viewport for the current level"),
-		[this](const FDoneDelegate& Done)
+	It(TEXT("CaptureAssetImage screenshots the viewport for the current level"),
+		[this]()
 		{
 			if (!FApp::CanEverRender())
 			{
 				// Rendering is disabled; skipping thumbnail test.
-				Done.Execute();
 				return;
 			}
 
@@ -207,24 +202,21 @@ void FEditorAppToolsetSpec::Define()
 			if (!EditorWorld)
 			{
 				// No editor world available; skipping viewport screenshot test.
-				Done.Execute();
 				return;
 			}
 
 			FString CurrentLevelPath = EditorWorld->GetOutermost()->GetName();
 			UToolCallAsyncResultImage* Result =
 				UEditorAppToolset::CaptureAssetImage(CurrentLevelPath);
-			TestNotNull("Result", Result);
-			if (!Result)
+			if (!TestNotNull("Result", Result))
 			{
-				Done.Execute();
 				return;
 			}
 
 			TStrongObjectPtr<UToolCallAsyncResultImage> StrongResult(Result);
 
 			AddCommand(new FFunctionLatentCommand(
-				[this, Done, StrongResult]() mutable -> bool
+				[this, StrongResult]() mutable -> bool
 				{
 					if (!StrongResult->bIsComplete)
 					{
@@ -234,7 +226,6 @@ void FEditorAppToolsetSpec::Define()
 					TestFalse("Has image data", StrongResult->Value.Data.IsEmpty());
 					TestEqual("MIME type",
 						StrongResult->Value.MimeType, FString(TEXT("image/png")));
-					Done.Execute();
 					return true; // done
 				}));
 		});
