[AI Assistant] Ignore error logs from side effects in Python tests

#Synced-CL 51924449

diff --git a/Engine/Plugins/Experimental/ToolsetRegistry/Source/ToolsetRegistry/Private/Tests/BlueprintLibraryToolsetTest.cpp b/Engine/Plugins/Experimental/ToolsetRegistry/Source/ToolsetRegistry/Private/Tests/BlueprintLibraryToolsetTest.cpp
--- a/Engine/Plugins/Experimental/ToolsetRegistry/Source/ToolsetRegistry/Private/Tests/BlueprintLibraryToolsetTest.cpp
+++ b/Engine/Plugins/Experimental/ToolsetRegistry/Source/ToolsetRegistry/Private/Tests/BlueprintLibraryToolsetTest.cpp
@@ -816,9 +816,14 @@ bool FBlueprintLibraryToolsetTestPythonTests::RunTest(
 		return true;
 	}
 
+	// Ignore all logged errors while running these tests. Some side effects cause errors to be
+	// logged that have nothing to do with the code under test, for example switching back to a
+	// level that has errors on load can report errors that aren't related to the code under test.
+	AddExpectedError(".*", EAutomationExpectedMessageFlags::MatchType::Contains, 1, true);
 	FPythonCommandEx PythonCommand = FPythonCommandEx();
 	PythonCommand.ExecutionMode = EPythonCommandExecutionMode::ExecuteStatement;
 	PythonCommand.Command = TEXT("import toolset_registry.tests; toolset_registry.tests.run_tests()");
+
 	return IPythonScriptPlugin::Get()->ExecPythonCommandEx(PythonCommand);
 }
 
