Add missing class docs to Agent Skills toolset.

#Synced-CL 51932664

diff --git a/Engine/Plugins/Experimental/ToolsetRegistry/Source/ToolsetRegistry/Public/ToolsetRegistry/AgentSkill.h b/Engine/Plugins/Experimental/ToolsetRegistry/Source/ToolsetRegistry/Public/ToolsetRegistry/AgentSkill.h
--- a/Engine/Plugins/Experimental/ToolsetRegistry/Source/ToolsetRegistry/Public/ToolsetRegistry/AgentSkill.h
+++ b/Engine/Plugins/Experimental/ToolsetRegistry/Source/ToolsetRegistry/Public/ToolsetRegistry/AgentSkill.h
@@ -64,7 +64,7 @@ protected:
 };
 
 
-
+/// Provides tools for listing, reading, and creating/updating skills.
 UCLASS(BlueprintType, Hidden)
 class UAgentSkillToolset : public UToolsetDefinition
 {
