Move ToolsetLibrary.h to public so external toolsets can use.

#Synced-CL 51941893

diff --git a/Engine/Plugins/Experimental/ToolsetRegistry/Source/ToolsetRegistry/Private/ToolsetRegistry/AgentSkill.cpp b/Engine/Plugins/Experimental/ToolsetRegistry/Source/ToolsetRegistry/Private/ToolsetRegistry/AgentSkill.cpp
--- a/Engine/Plugins/Experimental/ToolsetRegistry/Source/ToolsetRegistry/Private/ToolsetRegistry/AgentSkill.cpp
+++ b/Engine/Plugins/Experimental/ToolsetRegistry/Source/ToolsetRegistry/Private/ToolsetRegistry/AgentSkill.cpp
@@ -11,7 +11,7 @@
 
 #include "ToolsetRegistry/BlueprintLibraryToolset.h"
 #include "ToolsetRegistry/Toolset.h"
-#include "ToolsetLibrary.h"
+#include "ToolsetRegistry/ToolsetLibrary.h"
 #include "ToolsetRegistry/ToolsetRegistry.h"
 #include "ToolsetRegistry/ToolsetRegistrySubsystem.h"
 
diff --git a/Engine/Plugins/Experimental/ToolsetRegistry/Source/ToolsetRegistry/Private/ToolsetRegistry/ToolsetLibrary.cpp b/Engine/Plugins/Experimental/ToolsetRegistry/Source/ToolsetRegistry/Private/ToolsetRegistry/ToolsetLibrary.cpp
--- a/Engine/Plugins/Experimental/ToolsetRegistry/Source/ToolsetRegistry/Private/ToolsetRegistry/ToolsetLibrary.cpp
+++ b/Engine/Plugins/Experimental/ToolsetRegistry/Source/ToolsetRegistry/Private/ToolsetRegistry/ToolsetLibrary.cpp
@@ -1,6 +1,6 @@
 // Copyright Epic Games, Inc. All Rights Reserved.
 
-#include "ToolsetLibrary.h"
+#include "ToolsetRegistry/ToolsetLibrary.h"
 
 #include "AssetRegistry/AssetRegistryHelpers.h"
 #include "AssetRegistry/AssetRegistryModule.h"
diff --git a/Engine/Plugins/Experimental/ToolsetRegistry/Source/ToolsetRegistry/Private/ToolsetRegistry/ToolsetLibrary.h b/Engine/Plugins/Experimental/ToolsetRegistry/Source/ToolsetRegistry/Private/ToolsetRegistry/ToolsetLibrary.h
deleted file mode 100644
diff --git a/Engine/Plugins/Experimental/ToolsetRegistry/Source/ToolsetRegistry/Public/ToolsetRegistry/ToolsetLibrary.h b/Engine/Plugins/Experimental/ToolsetRegistry/Source/ToolsetRegistry/Public/ToolsetRegistry/ToolsetLibrary.h
new file mode 100644
--- /dev/null
+++ b/Engine/Plugins/Experimental/ToolsetRegistry/Source/ToolsetRegistry/Public/ToolsetRegistry/ToolsetLibrary.h
@@ -0,0 +1,74 @@
+// Copyright Epic Games, Inc. All Rights Reserved. 
+
+#pragma once
+
+#include "Kismet/BlueprintFunctionLibrary.h"
+
+#include "ToolsetRegistry/ToolsetImage.h"
+
+#include "ToolsetLibrary.generated.h"
+
+#define UE_API TOOLSETREGISTRY_API
+
+/// Provides functions that are critical to Python toolset operation but which are not available
+/// in and cannot be added to any other BFPL.
+UCLASS(BlueprintType, MinimalAPI)
+class UToolsetLibrary : public UBlueprintFunctionLibrary
+{
+	GENERATED_BODY()
+public:
+	/*
+	 * Returns the properties of a struct as JSON Schema.
+	 * By default only returns user-visible (Blueprint-accessible) properties.
+	 *
+	 * WARNING: bUserVisiblePropertiesOnly=false bypasses all visibility filtering and returns
+	 * every non-deprecated property on the struct. Only use this when you know that all struct
+	 * properties should be exposed. For example, DataTable row structs where every field is a
+	 * data column regardless of its Blueprint visibility flags.
+	 *
+	 * @param Struct The struct to extract properties from.
+	 * @param bUserVisiblePropertiesOnly When true (default), only Blueprint-accessible properties
+	 *        are included. When false, all non-deprecated properties are returned.
+	 * @return A JSON Schema formatted string that describes the properties.
+	 */
+	UFUNCTION(BlueprintCallable, Category = "Toolset Library")
+	static UE_API FString ListStructProperties(UStruct* Struct, bool bUserVisiblePropertiesOnly = true);
+
+	/*
+	 * Returns the values of the requested properties.
+	 * @param Object The object from which to extract properties.
+	 * @param PropertyNames The names of the properties to extract.
+	 * @return A JSON formatted string that contains the values of the requested properties.
+	 */
+	UFUNCTION(BlueprintCallable, Category = "Toolset Library")
+	static UE_API FString GetObjectProperties(const UObject* Object, const TArray<FName>& PropertyNames);
+
+	/*
+	 * Sets the values of specific properties in an Object.
+	 * @param Object The object to modify.
+	 * @param PropertiesJson The property names and values in a JSON formatted string.
+	 * @return True if all properties could be set.
+	 */
+	UFUNCTION(BlueprintCallable, Category = "Toolset Library")
+	static UE_API bool SetObjectProperties(UObject* Object, const FString& PropertiesJson);
+
+	/*
+	 * Returns the list of subclasses that derive from a class.
+	 * Uses the Asset Registry to find native and BP subclasses.
+	 * @param BaseClass The class to get subclasses from.
+	 * @return The list of native and BP classes that derive from the base class.
+	 */
+	UFUNCTION(BlueprintCallable, Category = "Toolset Library")
+	static UE_API TArray<UClass*> GetDerivedClasses(UClass* BaseClass);
+
+	/*
+	 * Returns the list of substructs that derive from a struct.
+	 * Iterates all loaded UScriptStruct objects to find matching substructs.
+	 * @param BaseStruct The struct to get substructs from.
+	 * @return The list of loaded structs that derive from the base struct.
+	 */
+	UFUNCTION(BlueprintCallable, Category = "Toolset Library")
+	static UE_API TArray<UScriptStruct*> GetDerivedStructs(UScriptStruct* BaseStruct);
+};
+
+#undef UE_API
