Disable render-dependent editor app toolset tests on build machine.

#Synced-CL 51975767

diff --git a/Engine/Plugins/Experimental/ToolsetRegistry/Source/ToolsetRegistry/Private/Tests/EditorAppToolsetTest.cpp b/Engine/Plugins/Experimental/ToolsetRegistry/Source/ToolsetRegistry/Private/Tests/EditorAppToolsetTest.cpp
--- a/Engine/Plugins/Experimental/ToolsetRegistry/Source/ToolsetRegistry/Private/Tests/EditorAppToolsetTest.cpp
+++ b/Engine/Plugins/Experimental/ToolsetRegistry/Source/ToolsetRegistry/Private/Tests/EditorAppToolsetTest.cpp
@@ -159,7 +159,7 @@ void FEditorAppToolsetSpec::Define()
 	It(TEXT("CaptureAssetImage renders a thumbnail for a static mesh"),
 		[this]()
 		{
-			if (!FApp::CanEverRender())
+			if (GIsBuildMachine || !FApp::CanEverRender())
 			{
 				// Rendering is disabled; skipping thumbnail test.
 				return;
@@ -192,7 +192,7 @@ void FEditorAppToolsetSpec::Define()
 	It(TEXT("CaptureAssetImage screenshots the viewport for the current level"),
 		[this]()
 		{
-			if (!FApp::CanEverRender())
+			if (GIsBuildMachine || !FApp::CanEverRender())
 			{
 				// Rendering is disabled; skipping thumbnail test.
 				return;
