Replace custom python toolset transform class with a tooslet converter. Simplifies toolset logic while giving us the idea LLM-facing behavior universally.

#Synced-CL 52052026

diff --git a/Engine/Plugins/Experimental/ToolsetRegistry/Source/ToolsetRegistry/Private/Tests/ToolsetJsonTest.cpp b/Engine/Plugins/Experimental/ToolsetRegistry/Source/ToolsetRegistry/Private/Tests/ToolsetJsonTest.cpp
--- a/Engine/Plugins/Experimental/ToolsetRegistry/Source/ToolsetRegistry/Private/Tests/ToolsetJsonTest.cpp
+++ b/Engine/Plugins/Experimental/ToolsetRegistry/Source/ToolsetRegistry/Private/Tests/ToolsetJsonTest.cpp
@@ -183,6 +183,29 @@ namespace
 		}
 	}
 
+	void TestTransformSchema(FAutomationTestBase& Test, const TSharedPtr<FJsonObject>& Schema)
+	{
+		Test.TestTrue("Schema", Schema.IsValid());
+		if (!Schema.IsValid()) return;
+
+		FString Type;
+		Test.TestTrue("Type", Schema->TryGetStringField(TEXT("type"), Type));
+		Test.TestEqual("Type", Type, FString(TEXT("object")));
+
+		Test.TestFalse("No required", Schema->HasField(TEXT("required")));
+
+		const TSharedPtr<FJsonObject>* Properties = nullptr;
+		Test.TestTrue("Properties", Schema->TryGetObjectField(TEXT("properties"), Properties));
+		if (Properties)
+		{
+			for (const FString& Field : TArray<FString>{ TEXT("location"), TEXT("rotation"), TEXT("scale") })
+			{
+				const TSharedPtr<FJsonObject>* SubSchema = nullptr;
+				Test.TestTrue(*Field, Properties->Get()->TryGetObjectField(Field, SubSchema));
+			}
+		}
+	}
+
 	TSharedPtr<FJsonObject> TestHasObjectField(
 		FAutomationTestBase& Test, const TSharedPtr<FJsonObject>& Schema,
 		const FString& RootObjectName, const TArrayView<FString>& FieldPath)
@@ -676,6 +699,75 @@ void FToolsetJsonSpec::Define()
 		TestEqual("ColorCheck", TestIn.TestLinearColor, TestOut.TestLinearColor);
 	});
 
+	It("Converts FTransform property to transform schema", [this]()
+	{
+		UStruct* Struct = FToolsetJsonTest::StaticStruct();
+		FProperty* Prop = Struct->FindPropertyByName("TestTransform");
+		const TSharedPtr<FJsonObject> Schema = PropertyToJsonSchema(Prop);
+		TestTransformSchema(*this, Schema);
+	});
+
+	It("Converts FTransform property to and from JSON", [this]()
+	{
+		FToolsetJsonTest TestIn;
+		TestIn.TestTransform = FTransform(
+			FRotator(10.f, 20.f, 30.f).Quaternion(),
+			FVector(100.f, 200.f, 300.f),
+			FVector(2.f, 3.f, 4.f));
+
+		UStruct* Struct = FToolsetJsonTest::StaticStruct();
+		FProperty* Property = Struct->FindPropertyByName("TestTransform");
+
+		TSharedPtr<FJsonValue> OutJson = PropertyToJsonData(Property, &TestIn.TestTransform);
+		if (!TestTrue("Json", OutJson.IsValid())) return;
+
+		TSharedPtr<FJsonObject> JsonObject = OutJson->AsObject();
+		if (!TestTrue("Json object", JsonObject.IsValid())) return;
+
+		TestTrue("location", JsonObject->HasField(TEXT("location")));
+		TestTrue("rotation", JsonObject->HasField(TEXT("rotation")));
+		TestTrue("scale", JsonObject->HasField(TEXT("scale")));
+
+		FToolsetJsonTest TestOut;
+		TestTrue("Output", JsonDataToProperty(OutJson, Property, &TestOut.TestTransform));
+		TestEqual("Location", TestOut.TestTransform.GetTranslation(),
+			TestIn.TestTransform.GetTranslation());
+		TestEqual("Rotation", TestOut.TestTransform.GetRotation().Rotator(),
+			TestIn.TestTransform.GetRotation().Rotator());
+		TestEqual("Scale", TestOut.TestTransform.GetScale3D(),
+			TestIn.TestTransform.GetScale3D());
+	});
+
+	It("Partial FTransform JSON preserves unset fields", [this]()
+	{
+		// Verify JsonDataToProperty with a partial JSON object (only location) leaves
+		// the existing rotation and scale unchanged — matching to_ue semantics.
+		const FRotator OriginalRotation(45.f, 90.f, 0.f);
+		const FVector OriginalScale(2.f, 2.f, 2.f);
+
+		FToolsetJsonTest Test;
+		Test.TestTransform = FTransform(
+			OriginalRotation.Quaternion(), FVector::ZeroVector, OriginalScale);
+
+		UStruct* Struct = FToolsetJsonTest::StaticStruct();
+		FProperty* Property = Struct->FindPropertyByName("TestTransform");
+
+		const FVector NewLocation(1.f, 2.f, 3.f);
+		TSharedPtr<FJsonObject> LocationJson = MakeShared<FJsonObject>();
+		LocationJson->SetNumberField(TEXT("x"), NewLocation.X);
+		LocationJson->SetNumberField(TEXT("y"), NewLocation.Y);
+		LocationJson->SetNumberField(TEXT("z"), NewLocation.Z);
+
+		TSharedPtr<FJsonObject> PartialJson = MakeShared<FJsonObject>();
+		PartialJson->SetObjectField(TEXT("location"), LocationJson);
+
+		TestTrue("Output", JsonDataToProperty(
+			MakeShared<FJsonValueObject>(PartialJson), Property, &Test.TestTransform));
+		TestEqual("Location updated", Test.TestTransform.GetTranslation(), NewLocation);
+		TestEqual("Rotation preserved", Test.TestTransform.GetRotation().Rotator(), OriginalRotation);
+		TestEqual("Scale preserved", Test.TestTransform.GetScale3D(), OriginalScale);
+	});
+
 	It("Converts struct to reference schema", [this]()
 	{		
 		TSharedPtr<FJsonObject> Schema = StructToJsonSchema(FToolsetJsonTest::StaticStruct());
@@ -701,6 +793,10 @@ void FToolsetJsonSpec::Define()
 					{
 						TestColorSchema(*this, *TestProperty, JsonPropertyName == FString(TEXT("testColor")));
 					}
+					else if (JsonPropertyName == TEXT("testTransform"))
+					{
+						TestTransformSchema(*this, *TestProperty);
+					}
 					else if (JsonPropertyName == TEXT("testToolCallAsyncResult"))
 					{
 						continue;  // Ignored as references to the base class are always null.
@@ -725,6 +821,10 @@ void FToolsetJsonSpec::Define()
 		TestIn.TestClass = UToolsetJsonTestObject::StaticClass();
 		TestIn.TestClassPath = UToolsetJsonTestObject::StaticClass();
 		TestIn.TestColor = FColor(10, 20, 30, 40);
+		TestIn.TestTransform = FTransform(
+			FRotator(10.f, 20.f, 30.f).Quaternion(),
+			FVector(100.f, 200.f, 300.f),
+			FVector(2.f, 3.f, 4.f));
 
 		TSharedPtr<FJsonObject> JsonObject = StructToJsonData(
 			FToolsetJsonTest::StaticStruct(), &TestIn);
@@ -747,6 +847,10 @@ void FToolsetJsonSpec::Define()
 				TestColorDataInStruct(
 					*this, JsonObject, JsonPropertyName, TestIn.TestLinearColor.ToFColor(false));
 			}
+			else if (JsonPropertyName == TEXT("testTransform"))
+			{
+				TestTrue("testTransform", JsonObject->HasField(TEXT("testTransform")));
+			}
 			else if (JsonPropertyName == TEXT("testToolCallAsyncResult"))
 			{
 				continue;  // Does not support conversion.
@@ -756,7 +860,7 @@ void FToolsetJsonSpec::Define()
 				TestObjectRefDataInStruct(
 					*this, JsonObject, JsonPropertyName, TestOther->GetPathName());
 			}
-		}		
+		}
 
 		FToolsetJsonTest TestOut;
 		TestTrue("FromJson", JsonDataToStruct(
@@ -768,6 +872,12 @@ void FToolsetJsonSpec::Define()
 		TestEqual("RefCheck", TestIn.TestClass, TestOut.TestClass);
 		TestEqual("RefCheck", TestIn.TestClassPath, TestOut.TestClassPath);
 		TestEqual("ColorCheck", TestIn.TestColor, TestOut.TestColor);
+		TestEqual("TransformLocation", TestIn.TestTransform.GetTranslation(),
+			TestOut.TestTransform.GetTranslation());
+		TestEqual("TransformRotation", TestIn.TestTransform.GetRotation().Rotator(),
+			TestOut.TestTransform.GetRotation().Rotator());
+		TestEqual("TransformScale", TestIn.TestTransform.GetScale3D(),
+			TestOut.TestTransform.GetScale3D());
 	});
 
 	It("Imports from old string style path", [this]()
diff --git a/Engine/Plugins/Experimental/ToolsetRegistry/Source/ToolsetRegistry/Private/Tests/ToolsetJsonTest.h b/Engine/Plugins/Experimental/ToolsetRegistry/Source/ToolsetRegistry/Private/Tests/ToolsetJsonTest.h
--- a/Engine/Plugins/Experimental/ToolsetRegistry/Source/ToolsetRegistry/Private/Tests/ToolsetJsonTest.h
+++ b/Engine/Plugins/Experimental/ToolsetRegistry/Source/ToolsetRegistry/Private/Tests/ToolsetJsonTest.h
@@ -116,6 +116,9 @@ public:
 	UPROPERTY()
 	FLinearColor TestLinearColor = FLinearColor::White;
 
+	UPROPERTY()
+	FTransform TestTransform;
+
 	UPROPERTY()
 	TObjectPtr<UToolCallAsyncResult> TestToolCallAsyncResult;
 };
diff --git a/Engine/Plugins/Experimental/ToolsetRegistry/Source/ToolsetRegistry/Private/ToolsetRegistry/ToolsetRegistry.cpp b/Engine/Plugins/Experimental/ToolsetRegistry/Source/ToolsetRegistry/Private/ToolsetRegistry/ToolsetRegistry.cpp
--- a/Engine/Plugins/Experimental/ToolsetRegistry/Source/ToolsetRegistry/Private/ToolsetRegistry/ToolsetRegistry.cpp
+++ b/Engine/Plugins/Experimental/ToolsetRegistry/Source/ToolsetRegistry/Private/ToolsetRegistry/ToolsetRegistry.cpp
@@ -14,6 +14,7 @@
 #include "ToolsetRegistry/ColorConverter.h"
 #include "ToolsetRegistry/Module.h"
 #include "ToolsetRegistry/ReferenceConverter.h"
+#include "ToolsetRegistry/TransformConverter.h"
 #include "ToolsetRegistry/Toolset.h"
 
 namespace UE::ToolsetRegistry
@@ -25,6 +26,9 @@ namespace UE::ToolsetRegistry
 
 		TSharedPtr<FToolsetReferenceConverter> ReferenceConverter = MakeShared<FToolsetReferenceConverter>();
 		RegisterConverter(ReferenceConverter);
+
+		TSharedPtr<FToolsetTransformConverter> TransformConverter = MakeShared<FToolsetTransformConverter>();
+		RegisterConverter(TransformConverter);
 	}
 
 	FToolsetRegistry::~FToolsetRegistry() = default;
diff --git a/Engine/Plugins/Experimental/ToolsetRegistry/Source/ToolsetRegistry/Private/ToolsetRegistry/TransformConverter.cpp b/Engine/Plugins/Experimental/ToolsetRegistry/Source/ToolsetRegistry/Private/ToolsetRegistry/TransformConverter.cpp
new file mode 100644
--- /dev/null
+++ b/Engine/Plugins/Experimental/ToolsetRegistry/Source/ToolsetRegistry/Private/ToolsetRegistry/TransformConverter.cpp
@@ -0,0 +1,101 @@
+// Copyright Epic Games, Inc. All Rights Reserved.
+
+#include "TransformConverter.h"
+
+#include "Dom/JsonValue.h"
+#include "Math/Transform.h"
+
+
+namespace UE::ToolsetRegistry
+{
+	FString FToolsetTransformConverter::GetName() const
+	{
+		static FString Name(TEXT("TransformConverter"));
+		return Name;
+	}
+
+	bool FToolsetTransformConverter::CanConvertProperty(TNotNull<const FProperty*> Property)
+	{
+		if (const FStructProperty* StructProperty = CastField<FStructProperty>(Property))
+		{
+			if (StructProperty->Struct == TBaseStructure<FTransform>::Get())
+			{
+				return true;
+			}
+		}
+		return false;
+	}
+
+	TSharedPtr<FJsonObject> FToolsetTransformConverter::PropertyToJsonSchema(
+		TNotNull<const FProperty*> Property)
+	{
+		return ToolsetStructToJsonSchema(FToolsetTransform::StaticStruct());
+	}
+
+	TSharedPtr<FJsonValue> FToolsetTransformConverter::PropertyToDefault(
+		TNotNull<const FProperty*> Property, const FString& DefaultString)
+	{
+		FTransform DefaultTransform = FTransform::Identity;
+		if (!DefaultString.IsEmpty())
+		{
+			Property->ImportText_Direct(*DefaultString, &DefaultTransform, nullptr, PPF_None);
+		}
+		FToolsetTransform ToolsetTransform = FromUE(DefaultTransform);
+		return MakeShared<FJsonValueObject>(
+			ToolsetStructToJsonData(FToolsetTransform::StaticStruct(), &ToolsetTransform));
+	}
+
+	TSharedPtr<FJsonValue> FToolsetTransformConverter::PropertyToJsonData(
+		TNotNull<FProperty*> Property, const void* Value)
+	{
+		const FTransform& Transform = *static_cast<const FTransform*>(Value);
+		FToolsetTransform ToolsetTransform = FromUE(Transform);
+		return MakeShared<FJsonValueObject>(
+			ToolsetStructToJsonData(FToolsetTransform::StaticStruct(), &ToolsetTransform));
+	}
+
+	bool FToolsetTransformConverter::JsonDataToProperty(
+		const TSharedPtr<FJsonValue>& JsonValue, TNotNull<FProperty*> Property,
+		void* OutValue)
+	{
+		TSharedPtr<FJsonObject> JsonObject = JsonValue->AsObject();
+		if (!JsonObject.IsValid())
+		{
+			return false;
+		}
+		FTransform& OutTransform = *static_cast<FTransform*>(OutValue);
+		FToolsetTransform ToolsetTransform;
+		ToolsetJsonDataToStruct(
+			JsonObject.ToSharedRef(), FToolsetTransform::StaticStruct(), &ToolsetTransform);
+		OutTransform = ToUE(ToolsetTransform, OutTransform);
+		return true;
+	}
+
+	FTransform FToolsetTransformConverter::ToUE(
+		const FToolsetTransform& Transform, const FTransform& Base)
+	{
+		FTransform Result = Base;
+		if (Transform.Location.IsSet())
+		{
+			Result.SetTranslation(Transform.Location.GetValue());
+		}
+		if (Transform.Rotation.IsSet())
+		{
+			Result.SetRotation(Transform.Rotation.GetValue().Quaternion());
+		}
+		if (Transform.Scale.IsSet())
+		{
+			Result.SetScale3D(Transform.Scale.GetValue());
+		}
+		return Result;
+	}
+
+	FToolsetTransform FToolsetTransformConverter::FromUE(const FTransform& Transform)
+	{
+		FToolsetTransform Result;
+		Result.Location = Transform.GetTranslation();
+		Result.Rotation = Transform.GetRotation().Rotator();
+		Result.Scale = Transform.GetScale3D();
+		return Result;
+	}
+}
diff --git a/Engine/Plugins/Experimental/ToolsetRegistry/Source/ToolsetRegistry/Private/ToolsetRegistry/TransformConverter.h b/Engine/Plugins/Experimental/ToolsetRegistry/Source/ToolsetRegistry/Private/ToolsetRegistry/TransformConverter.h
new file mode 100644
--- /dev/null
+++ b/Engine/Plugins/Experimental/ToolsetRegistry/Source/ToolsetRegistry/Private/ToolsetRegistry/TransformConverter.h
@@ -0,0 +1,63 @@
+// Copyright Epic Games, Inc. All Rights Reserved.
+
+#pragma once
+
+#include "CoreMinimal.h"
+#include "ToolsetRegistry/ToolsetJsonConverter.h"
+
+#include "TransformConverter.generated.h"
+
+/// Represents a 3D transformation with optional location, rotation, and scale.
+/// Unset fields mean "identity" when creating objects and "don't change" when modifying existing ones.
+USTRUCT(BlueprintType, MinimalAPI)
+struct FToolsetTransform
+{
+	GENERATED_BODY()
+public:
+	/// The world-space location.
+	UPROPERTY(BlueprintReadWrite, Category = "Transform")
+	TOptional<FVector> Location;
+
+	/// The world-space rotation.
+	UPROPERTY(BlueprintReadWrite, Category = "Transform")
+	TOptional<FRotator> Rotation;
+
+	/// The scale.
+	UPROPERTY(BlueprintReadWrite, Category = "Transform")
+	TOptional<FVector> Scale;
+};
+
+namespace UE::ToolsetRegistry
+{
+	/// Converts FTransform properties to/from JSON using optional translation, rotation, and scale
+	/// fields, matching the semantics of the Python Transform struct's to_ue / from_ue methods.
+	class FToolsetTransformConverter : public FToolsetJsonConverter
+	{
+	public:
+		virtual FString GetName() const override;
+
+		virtual bool CanConvertProperty(TNotNull<const FProperty*> Property) override;
+
+		virtual TSharedPtr<FJsonObject> PropertyToJsonSchema(
+			TNotNull<const FProperty*> Property) override;
+
+		virtual TSharedPtr<FJsonValue> PropertyToDefault(
+			TNotNull<const FProperty*> Property, const FString& DefaultString) override;
+
+		virtual TSharedPtr<FJsonValue> PropertyToJsonData(
+			TNotNull<FProperty*> Property, const void* Value) override;
+
+		virtual bool JsonDataToProperty(
+			const TSharedPtr<FJsonValue>& JsonValue, TNotNull<FProperty*> Property,
+			void* OutValue) override;
+
+	private:
+		// Equivalent to Python's to_ue: applies set fields on top of Base, leaves unset fields
+		// at their Base values.
+		static FTransform ToUE(const FToolsetTransform& Transform,
+			const FTransform& Base = FTransform::Identity);
+
+		// Equivalent to Python's from_ue: always populates all three fields from the transform.
+		static FToolsetTransform FromUE(const FTransform& Transform);
+	};
+}
