[Toolset Registry] Fix memory corruption in image capture tools.

`FViewportScreenshotCapture` unregistered the screenshot delegate handle
within `OnCaptured()` which caused the lamba holding the only reference
to `FViewportScreenshotCapture` to be deleted which deleted
`FViewportScreenshotCapture` (i.e this) within `OnCaptured()` causing
memory corruption when setting the result.

While in there:
*  `Capture` objects have been changed to use TSharedFromThis which ensures
- at runtime - they cannot be created from a non-shared pointer.
*  Used `TStrongObjectPtr` to ensure references are implicitly released on
`Capture` object destruction.
* Used `FDelegateHandleRaii` to ensure the registered delegate is always
released on  `FViewportScreenshotCapture` object destruction.
* Try re-enabling image capture tests in automated builds.

#Synced-CL 52024537

diff --git a/Engine/Plugins/Experimental/ToolsetRegistry/Source/ToolsetRegistry/Private/Tests/EditorAppToolsetTest.cpp b/Engine/Plugins/Experimental/ToolsetRegistry/Source/ToolsetRegistry/Private/Tests/EditorAppToolsetTest.cpp
--- a/Engine/Plugins/Experimental/ToolsetRegistry/Source/ToolsetRegistry/Private/Tests/EditorAppToolsetTest.cpp
+++ b/Engine/Plugins/Experimental/ToolsetRegistry/Source/ToolsetRegistry/Private/Tests/EditorAppToolsetTest.cpp
@@ -159,7 +159,7 @@ void FEditorAppToolsetSpec::Define()
 	It(TEXT("CaptureAssetImage renders a thumbnail for a static mesh"),
 		[this]()
 		{
-			if (GIsBuildMachine || !FApp::CanEverRender())
+			if (!FApp::CanEverRender())
 			{
 				// Rendering is disabled; skipping thumbnail test.
 				return;
@@ -192,7 +192,7 @@ void FEditorAppToolsetSpec::Define()
 	It(TEXT("CaptureAssetImage screenshots the viewport for the current level"),
 		[this]()
 		{
-			if (GIsBuildMachine || !FApp::CanEverRender())
+			if (!FApp::CanEverRender())
 			{
 				// Rendering is disabled; skipping thumbnail test.
 				return;
diff --git a/Engine/Plugins/Experimental/ToolsetRegistry/Source/ToolsetRegistry/Private/ToolsetRegistry/EditorAppToolset.cpp b/Engine/Plugins/Experimental/ToolsetRegistry/Source/ToolsetRegistry/Private/ToolsetRegistry/EditorAppToolset.cpp
--- a/Engine/Plugins/Experimental/ToolsetRegistry/Source/ToolsetRegistry/Private/ToolsetRegistry/EditorAppToolset.cpp
+++ b/Engine/Plugins/Experimental/ToolsetRegistry/Source/ToolsetRegistry/Private/ToolsetRegistry/EditorAppToolset.cpp
@@ -21,10 +21,14 @@
 #include "ObjectTools.h"
 #include "RenderingThread.h"
 #include "ShaderCompiler.h"
+#include "Templates/SharedPointer.h"
+#include "ToolsetRegistry/DelegateHandle.h"
 #include "ToolsetRegistry/JsonConversion.h"
 #include "ToolsetRegistry/ToolCallAsyncResultImage.h"
 #include "ToolsetRegistry/ToolsetImage.h"
+#include "UObject/StrongObjectPtr.h"
 #include "UObject/UObjectGlobals.h"
+
 #include "UnrealClient.h"
 
 #include UE_INLINE_GENERATED_CPP_BY_NAME(EditorAppToolset)
@@ -34,36 +38,38 @@ namespace
 	// Helper that captures a single viewport screenshot and resolves a result.
 	// The object is kept alive by a strong reference captured in the screenshot
 	// delegate lambda, which is released once the callback fires.
-	struct FViewportScreenshotCapture
+	struct FViewportScreenshotCapture : public TSharedFromThis<FViewportScreenshotCapture>
 	{
-		FViewportScreenshotCapture()
-		{
-			Result = NewObject<UToolCallAsyncResultImage>();
-			Result->AddToRoot();
-		}
+	public:
+		FViewportScreenshotCapture() : Result(NewObject<UToolCallAsyncResultImage>()) {}
+		~FViewportScreenshotCapture() = default;
 
-		static UToolCallAsyncResultImage* Create()
+		static TStrongObjectPtr<UToolCallAsyncResultImage> Start()
 		{
-			TSharedRef<FViewportScreenshotCapture> Self = MakeShared<FViewportScreenshotCapture>();
-			Self->Handle = FScreenshotRequest::OnScreenshotCaptured().AddLambda(
-				[Self](int32 Width, int32 Height, const TArray<FColor>& Colors)
-				{
-					Self->OnCaptured(Width, Height, Colors);
-				});
-			FScreenshotRequest::RequestScreenshot(false);
-			return Self->Result;
+			return MakeShared<FViewportScreenshotCapture>()->StartCapture();
 		}
 
-		~FViewportScreenshotCapture()
+	private:
+		TStrongObjectPtr<UToolCallAsyncResultImage> StartCapture()
 		{
-			Result->RemoveFromRoot();
+			FOnScreenshotCaptured& OnScreenshotCaptured =
+				FScreenshotRequest::OnScreenshotCaptured();
+			Handle =
+				UE::ToolsetRegistry::FDelegateHandleRaii::Create(
+					OnScreenshotCaptured,
+					OnScreenshotCaptured.AddLambda(
+						[This = AsShared().ToSharedPtr()](
+							int32 Width, int32 Height, const TArray<FColor>& Colors) mutable
+						{
+							This->OnCaptured(Width, Height, Colors);
+							This.Reset();
+						}));
+			FScreenshotRequest::RequestScreenshot(false);
+			return Result;
 		}
 
-	private:
 		void OnCaptured(int32 Width, int32 Height, const TArray<FColor>& Colors)
 		{
-			FScreenshotRequest::OnScreenshotCaptured().Remove(Handle);
-
 			const int32 MaxSize = 1280;
 			FIntPoint Dimensions(Width, Height);
 			TArray<FColor> ScaledPixels;
@@ -96,45 +102,47 @@ namespace
 			}
 		}
 
-		UToolCallAsyncResultImage* Result;
-		FDelegateHandle Handle;
+		TStrongObjectPtr<UToolCallAsyncResultImage> Result;
+		UE::ToolsetRegistry::FDelegateHandleRaii Handle;
 	};
 
 	// Helper that renders an asset thumbnail and resolves a result.
 	// Initiates an async load then polls each tick until the asset is loaded and
 	// all compilation (assets and shaders) is complete before rendering.
-	struct FAssetThumbnailCapture
+	struct FAssetThumbnailCapture : public TSharedFromThis<FAssetThumbnailCapture>
 	{
-		explicit FAssetThumbnailCapture(const FString& InAssetPath)
-			: AssetPath(InAssetPath)
-			, StartTime(FPlatformTime::Seconds())
+		explicit FAssetThumbnailCapture(const FString& InAssetPath) :
+			AssetPath(InAssetPath),
+			Result(NewObject<UToolCallAsyncResultImage>()),
+			StartTime(FPlatformTime::Seconds())
 		{
-			Result = NewObject<UToolCallAsyncResultImage>();
-			Result->AddToRoot();
 		}
+		~FAssetThumbnailCapture() = default;
 
-		static UToolCallAsyncResultImage* Create(const FString& InAssetPath)
+		static TStrongObjectPtr<UToolCallAsyncResultImage> Start(const FString& InAssetPath)
 		{
-			TSharedRef<FAssetThumbnailCapture> Self = MakeShared<FAssetThumbnailCapture>(InAssetPath);
-			// Initiate async load; the ticker polls for completion.
-			Self->StreamableHandle = UAssetManager::Get().GetStreamableManager().RequestAsyncLoad(
-				FSoftObjectPath(Self->AssetPath), FStreamableDelegate());
-			Self->TickerHandle = FTSTicker::GetCoreTicker().AddTicker(
-				FTickerDelegate::CreateLambda([Self](float) -> bool
-				{
-					return Self->OnTick(Self);
-				}));
-			return Self->Result;
+			return MakeShared<FAssetThumbnailCapture>(InAssetPath)->StartCapture();
 		}
 
-		~FAssetThumbnailCapture()
+	private:
+		TStrongObjectPtr<UToolCallAsyncResultImage> StartCapture()
 		{
-			Result->RemoveFromRoot();
+			// Initiate async load; the ticker polls for completion.
+			StreamableHandle = UAssetManager::Get().GetStreamableManager().RequestAsyncLoad(
+				FSoftObjectPath(AssetPath), FStreamableDelegate());
+			TickerHandle = FTSTicker::GetCoreTicker().AddTicker(
+				FTickerDelegate::CreateLambda(
+					[This = AsShared().ToSharedPtr()](float) mutable -> bool
+					{
+						bool bContinue = This->OnTick();
+						if (!bContinue) This.Reset();
+						return bContinue;
+					}));
+			return Result;
 		}
 
-	private:
 		// Returns true to keep ticking, false when done.
-		bool OnTick(TSharedRef<FAssetThumbnailCapture> Self)
+		bool OnTick()
 		{
 			if (FPlatformTime::Seconds() - StartTime > 10.0)
 			{
@@ -233,7 +241,7 @@ namespace
 		}
 
 		FString AssetPath;
-		UToolCallAsyncResultImage* Result;
+		TStrongObjectPtr<UToolCallAsyncResultImage> Result;
 		TSharedPtr<FStreamableHandle> StreamableHandle;
 		FTSTicker::FDelegateHandle TickerHandle;
 		double StartTime;
@@ -304,7 +312,7 @@ UToolCallAsyncResultImage* UEditorAppToolset::CaptureAssetImage(const FString& A
 	if (AssetPath.IsEmpty() ||
 		CurrentLevelPackage == FPackageName::ObjectPathToPackageName(AssetPath))
 	{
-		return FViewportScreenshotCapture::Create();
+		return FViewportScreenshotCapture::Start().Get();
 	}
 
 	// Thumbnail rendering path
@@ -344,7 +352,7 @@ UToolCallAsyncResultImage* UEditorAppToolset::CaptureAssetImage(const FString& A
 	}
 	else
 	{
-		Result = FAssetThumbnailCapture::Create(ObjectPath);
+		Result = FAssetThumbnailCapture::Start(ObjectPath).Get();
 	}
 	return Result.Get();
 }
