Split up ToolsetJsonTest into individual tests per converter w/ their own header files. Also gave ToolCallScriptExceptionHandler its own header.

#Synced-CL 52087183

diff --git a/Engine/Plugins/Experimental/ToolsetRegistry/Source/ToolsetRegistry/Private/Tests/ColorConverterTest.cpp b/Engine/Plugins/Experimental/ToolsetRegistry/Source/ToolsetRegistry/Private/Tests/ColorConverterTest.cpp
new file mode 100644
--- /dev/null
+++ b/Engine/Plugins/Experimental/ToolsetRegistry/Source/ToolsetRegistry/Private/Tests/ColorConverterTest.cpp
@@ -0,0 +1,178 @@
+// Copyright Epic Games, Inc. All Rights Reserved.
+
+#include "ColorConverterTest.h"
+
+#include "CoreMinimal.h"
+#include "Misc/AutomationTest.h"
+
+#include "ToolsetRegistry/JsonSchema.h"
+#include "ToolsetRegistry/ToolsetJson.h"
+
+#if WITH_DEV_AUTOMATION_TESTS
+
+BEGIN_DEFINE_SPEC(FToolsetColorConverterSpec, "AI.ToolsetRegistry.ColorConverterSpec",
+	EAutomationTestFlags::EditorContext | EAutomationTestFlags::EngineFilter)
+	void TestColorSchema(const TSharedPtr<FJsonObject>& Schema, bool HasMax);
+	void TestColorData(const TSharedPtr<FJsonObject>& Schema, const FColor& Expected);
+END_DEFINE_SPEC(FToolsetColorConverterSpec)
+
+void FToolsetColorConverterSpec::TestColorSchema(const TSharedPtr<FJsonObject>& Schema, bool HasMax)
+{
+	if (!TestTrue("Schema", Schema.IsValid())) return;
+
+	FString Type;
+	TestTrue("Type", Schema->TryGetStringField(TEXT("type"), Type));
+	TestEqual("Type", Type, FString(TEXT("object")));
+
+	FString Title;
+	TestTrue("Title", Schema->TryGetStringField(TEXT("title"), Title));
+	TestEqual("Title", Title, FString(TEXT("LinearColor")));
+
+	TestFalse("Description", Schema->HasField(TEXT("description")));
+
+	const TSharedPtr<FJsonObject>* Properties = nullptr;
+	if (!(TestTrue("Properties", Schema->TryGetObjectField(TEXT("properties"), Properties)) &&
+		Properties))
+	{
+		return;
+	}
+	TArray<FString> PropertyNames({ "r", "g", "b", "a" });
+	for (const auto& Name : PropertyNames)
+	{
+		const TSharedPtr<FJsonObject>* Property = nullptr;
+		if (!(TestTrue("Property", Properties->Get()->TryGetObjectField(Name, Property)) &&
+			Property))
+		{
+			continue;
+		}
+		TestEqual("Property", Property->Get()->GetStringField(
+			TEXT("type")), FString(TEXT("number")));
+		TestEqual("Property", Property->Get()->GetNumberField(
+			TEXT("minimum")), 0.0);
+		if (HasMax)
+		{
+			TestEqual("Property", Property->Get()->GetNumberField(
+				TEXT("maximum")), 1.0);
+		}
+	}
+}
+
+void FToolsetColorConverterSpec::TestColorData(
+	const TSharedPtr<FJsonObject>& Schema, const FColor& Expected)
+{
+	if (!TestTrue("Schema", Schema.IsValid())) return;
+
+	const FLinearColor LinearExpected = Expected.ReinterpretAsLinear();
+
+	float R;
+	TestTrue("Color", Schema->TryGetNumberField(TEXT("r"), R));
+	TestEqual("Color", R, LinearExpected.R, 2e-3f);
+
+	float G;
+	TestTrue("Color", Schema->TryGetNumberField(TEXT("g"), G));
+	TestEqual("Color", G, LinearExpected.G, 2e-3f);
+
+	float B;
+	TestTrue("Color", Schema->TryGetNumberField(TEXT("b"), B));
+	TestEqual("Color", B, LinearExpected.B, 2e-3f);
+
+	float A;
+	TestTrue("Color", Schema->TryGetNumberField(TEXT("a"), A));
+	TestEqual("Color", A, LinearExpected.A, 2e-3f);
+}
+
+void FToolsetColorConverterSpec::Define()
+{
+	using namespace UE::ToolsetRegistry::Internal;
+	using namespace UE::ToolsetRegistry::Internal::ToolsetJson;
+
+	It("Converts FColor property to FLinearColor schema", [this]()
+	{
+		UStruct* Struct = FColorConverterTest::StaticStruct();
+		FProperty* Prop = Struct->FindPropertyByName("TestColor");
+		const TSharedPtr<FJsonObject> Schema = PropertyToJsonSchema(Prop);
+		TestColorSchema(Schema, true);
+	});
+
+	It("Converts FColor property to and from JSON", [this]()
+	{
+		FColorConverterTest TestIn;
+		TestIn.TestColor = FColor(10, 20, 30, 40);
+
+		UStruct* Struct = FColorConverterTest::StaticStruct();
+		FProperty* Property = Struct->FindPropertyByName("TestColor");
+
+		TSharedPtr<FJsonValue> OutJson = PropertyToJsonData(Property, &TestIn.TestColor);
+		TestColorData(OutJson->AsObject(), TestIn.TestColor);
+
+		FColorConverterTest TestOut;
+		TestTrue("Output", JsonDataToProperty(OutJson, Property, &TestOut.TestColor));
+		TestEqual("ColorCheck", TestIn.TestColor, TestOut.TestColor);
+	});
+
+	It("Converts FLinearColor property to FLinearColor schema", [this]()
+	{
+		UStruct* Struct = FColorConverterTest::StaticStruct();
+		FProperty* Prop = Struct->FindPropertyByName("TestLinearColor");
+		const TSharedPtr<FJsonObject> Schema = PropertyToJsonSchema(Prop);
+		TestColorSchema(Schema, false);
+	});
+
+	It("Converts FLinearColor property to and from JSON", [this]()
+	{
+		FColorConverterTest TestIn;
+		TestIn.TestLinearColor = FLinearColor(0.1, 0.2, 0.3, 0.4);
+
+		UStruct* Struct = FColorConverterTest::StaticStruct();
+		FProperty* Property = Struct->FindPropertyByName("TestLinearColor");
+
+		TSharedPtr<FJsonValue> OutJson = PropertyToJsonData(Property, &TestIn.TestLinearColor);
+		TestColorData(OutJson->AsObject(), TestIn.TestLinearColor.ToFColor(false));
+
+		FColorConverterTest TestOut;
+		TestTrue("Output", JsonDataToProperty(OutJson, Property, &TestOut.TestLinearColor));
+		TestEqual("ColorCheck", TestIn.TestLinearColor, TestOut.TestLinearColor);
+	});
+
+	It("Can read color defaults", [this]()
+	{
+		UFunction* FunctionPtr = UColorConverterTestObject::StaticClass()->FindFunctionByName("TestDefault");
+		const TSharedPtr<FJsonObject> Schema = StructToJsonSchema(FunctionPtr);
+		if (!TestTrue("Schema", Schema.IsValid())) return;
+
+		const TSharedPtr<FJsonObject>* InputSchema = nullptr;
+		if (!(TestTrue("Input", Schema->TryGetObjectField(TEXT("inputSchema"), InputSchema)) &&
+			InputSchema))
+		{
+			return;
+		}
+
+		TestFalse("NoRequired", InputSchema->Get()->HasField(TEXT("required")));
+
+		const TSharedPtr<FJsonObject>* Properties = nullptr;
+		if (!(TestTrue("Properties", InputSchema->Get()->TryGetObjectField(TEXT("properties"), Properties)) &&
+			Properties))
+		{
+			return;
+		}
+
+		const TSharedPtr<FJsonObject>* TestColor = nullptr;
+		if (!(TestTrue("testColor", Properties->Get()->TryGetObjectField(TEXT("testColor"), TestColor)) &&
+			TestColor))
+		{
+			return;
+		}
+		TestColorData(TestColor->Get()->GetObjectField(TEXT("default")), FColor(10, 20, 30, 40));
+
+		const TSharedPtr<FJsonObject>* TestLinearColor = nullptr;
+		if (!(TestTrue("testLinearColor", Properties->Get()->TryGetObjectField(TEXT("testLinearColor"), TestLinearColor)) &&
+			TestLinearColor))
+		{
+			return;
+		}
+		TestColorData(
+			TestLinearColor->Get()->GetObjectField(TEXT("default")), FLinearColor(0.1, 0.2, 0.3, 0.4).ToFColor(false));
+	});
+}
+
+#endif
diff --git a/Engine/Plugins/Experimental/ToolsetRegistry/Source/ToolsetRegistry/Private/Tests/ColorConverterTest.h b/Engine/Plugins/Experimental/ToolsetRegistry/Source/ToolsetRegistry/Private/Tests/ColorConverterTest.h
new file mode 100644
--- /dev/null
+++ b/Engine/Plugins/Experimental/ToolsetRegistry/Source/ToolsetRegistry/Private/Tests/ColorConverterTest.h
@@ -0,0 +1,31 @@
+// Copyright Epic Games, Inc. All Rights Reserved.
+
+#pragma once
+
+#include "CoreMinimal.h"
+
+#include "ColorConverterTest.generated.h"
+
+
+UCLASS(BlueprintType)
+class UColorConverterTestObject : public UObject
+{
+	GENERATED_BODY()
+public:
+	UFUNCTION(BlueprintCallable, Category="ColorConverterTest")
+	void TestDefault(
+		FColor TestColor = FColor(10, 20, 30, 40),
+		FLinearColor TestLinearColor = FLinearColor(0.1, 0.2, 0.3, 0.4)) {}
+};
+
+USTRUCT(BlueprintType)
+struct FColorConverterTest
+{
+	GENERATED_BODY()
+public:
+	UPROPERTY()
+	FColor TestColor = FColor::White;
+
+	UPROPERTY()
+	FLinearColor TestLinearColor = FLinearColor::White;
+};
diff --git a/Engine/Plugins/Experimental/ToolsetRegistry/Source/ToolsetRegistry/Private/Tests/ReferenceConverterTest.cpp b/Engine/Plugins/Experimental/ToolsetRegistry/Source/ToolsetRegistry/Private/Tests/ReferenceConverterTest.cpp
new file mode 100644
--- /dev/null
+++ b/Engine/Plugins/Experimental/ToolsetRegistry/Source/ToolsetRegistry/Private/Tests/ReferenceConverterTest.cpp
@@ -0,0 +1,518 @@
+// Copyright Epic Games, Inc. All Rights Reserved.
+
+#include "ReferenceConverterTest.h"
+
+#include "AssetRegistry/AssetData.h"
+#include "AssetRegistry/AssetRegistryModule.h"
+#include "CoreMinimal.h"
+#include "Engine/Blueprint.h"
+#include "Misc/AutomationTest.h"
+
+#include "ToolsetRegistry/JsonSchema.h"
+#include "ToolsetRegistry/ToolsetJson.h"
+#include "Tests/ToolCallTestHelpers.h"
+
+#if WITH_DEV_AUTOMATION_TESTS
+
+BEGIN_DEFINE_SPEC(FToolsetReferenceConverterSpec, "AI.ToolsetRegistry.ReferenceConverterSpec",
+	EAutomationTestFlags::EditorContext | EAutomationTestFlags::EngineFilter)
+	void TestObjectRefSchema(const TSharedPtr<FJsonObject>& Schema, UClass* Class, bool IsClass);
+	void TestObjectRefData(const TSharedPtr<FJsonObject>& Schema, const FString& Expected);
+END_DEFINE_SPEC(FToolsetReferenceConverterSpec)
+
+void FToolsetReferenceConverterSpec::TestObjectRefSchema(
+	const TSharedPtr<FJsonObject>& Schema, UClass* Class, bool IsClass)
+{
+	if (!TestTrue("Schema", Schema.IsValid())) return;
+
+	FString Type;
+	TestTrue("Type", Schema->TryGetStringField(TEXT("type"), Type));
+	TestEqual("Type", Type, FString(TEXT("object")));
+
+	if (Class)
+	{
+		FString Title;
+		TestTrue("Title", Schema->TryGetStringField(TEXT("title"), Title));
+		FString ClassPrefix(IsClass ? TEXT("Class@") : TEXT(""));
+		TestEqual("Title", Title, ClassPrefix + Class->GetPathName());
+	}
+
+	const TSharedPtr<FJsonObject>* Properties = nullptr;
+	if (!(TestTrue("Properties", Schema->TryGetObjectField(TEXT("properties"), Properties)) &&
+		Properties))
+	{
+		return;
+	}
+	const TSharedPtr<FJsonObject>* RefPath = nullptr;
+	if (!(TestTrue("RefPath", Properties->Get()->TryGetObjectField(TEXT("refPath"), RefPath)) &&
+		RefPath))
+	{
+		return;
+	}
+	TestTrue("RefPath", RefPath->Get()->TryGetStringField(TEXT("type"), Type));
+	TestEqual("RefPath", Type, FString(TEXT("string")));
+}
+
+void FToolsetReferenceConverterSpec::TestObjectRefData(
+	const TSharedPtr<FJsonObject>& Schema, const FString& Expected)
+{
+	if (!TestTrue("Schema", Schema.IsValid())) return;
+
+	FString RefPath;
+	TestTrue("Unpack", Schema->TryGetStringField(TEXT("refPath"), RefPath));
+	TestEqual("RefPath", RefPath, Expected);
+}
+
+void FToolsetReferenceConverterSpec::Define()
+{
+	using namespace UE::ToolsetRegistry::Internal;
+	using namespace UE::ToolsetRegistry::Internal::ToolsetJson;
+
+	It("Converts UObject* property to reference schema", [this]()
+	{
+		UStruct* Struct = FReferenceConverterTest::StaticStruct();
+		FProperty* Prop = Struct->FindPropertyByName("TestObject");
+		const TSharedPtr<FJsonObject> Schema = PropertyToJsonSchema(Prop);
+		TestObjectRefSchema(Schema, UReferenceConverterTestObject::StaticClass(), false);
+	});
+
+	It("Converts UObject* property to and from JSON", [this]()
+	{
+		FReferenceConverterTest TestIn;
+		TObjectPtr<UReferenceConverterTestObject> TestOther = NewObject<UReferenceConverterTestObject>();
+		TestIn.TestObject = TestOther;
+
+		UStruct* Struct = FReferenceConverterTest::StaticStruct();
+		FProperty* Property = Struct->FindPropertyByName("TestObject");
+
+		TSharedPtr<FJsonValue> OutJson = PropertyToJsonData(Property, &TestIn.TestObject);
+		TestObjectRefData(OutJson->AsObject(), TestOther->GetPathName());
+
+		FReferenceConverterTest TestOut;
+		TestTrue("Output", JsonDataToProperty(OutJson, Property, &TestOut.TestObject));
+		TestEqual("RefCheck", TestIn.TestObject, TestOut.TestObject);
+	});
+
+	It("Ensures refs can be set from all forms of null", [this]()
+	{
+		FReferenceConverterTest Test;
+		TObjectPtr<UReferenceConverterTestObject> TestOther = NewObject<UReferenceConverterTestObject>();
+		UStruct* Struct = FReferenceConverterTest::StaticStruct();
+		FProperty* Property = Struct->FindPropertyByName("TestObject");
+		TSharedPtr<FJsonObject> JsonRef = MakeShared<FJsonObject>();
+
+		// Null Json type.
+		Test.TestObject = TestOther;
+		TestEqual("RefCheck", Test.TestObject, TestOther);
+		JsonRef->SetField(FString(TEXT("refPath")), MakeShared<FJsonValueNull>());
+		TestTrue("Output", JsonDataToProperty(
+			MakeShared<FJsonValueObject>(JsonRef), Property, &Test.TestObject));
+		TestEqual("RefCheck", Test.TestObject, TObjectPtr<UReferenceConverterTestObject>());
+
+		// None is the UE standard.
+		Test.TestObject = TestOther;
+		TestEqual("RefCheck", Test.TestObject, TestOther);
+		JsonRef->SetStringField(FString(TEXT("refPath")), FString(TEXT("None")));
+		TestTrue("Output", JsonDataToProperty(
+			MakeShared<FJsonValueObject>(JsonRef), Property, &Test.TestObject));
+		TestEqual("RefCheck", Test.TestObject, TObjectPtr<UReferenceConverterTestObject>());
+
+		// Empty string also means null.
+		Test.TestObject = TestOther;
+		TestEqual("RefCheck", Test.TestObject, TestOther);
+		JsonRef->SetStringField(FString(TEXT("refPath")), FString(TEXT("")));
+		TestTrue("Output", JsonDataToProperty(
+			MakeShared<FJsonValueObject>(JsonRef), Property, &Test.TestObject));
+		TestEqual("RefCheck", Test.TestObject, TObjectPtr<UReferenceConverterTestObject>());
+
+		// And we also support null just in case.
+		Test.TestObject = TestOther;
+		TestEqual("RefCheck", Test.TestObject, TestOther);
+		JsonRef->SetStringField(FString(TEXT("refPath")), FString(TEXT("null")));
+		TestTrue("Output", JsonDataToProperty(
+			MakeShared<FJsonValueObject>(JsonRef), Property, &Test.TestObject));
+		TestEqual("RefCheck", Test.TestObject, TObjectPtr<UReferenceConverterTestObject>());
+	});
+
+	It("Accepts null for optional parameter", [this]()
+	{
+		TObjectPtr<UReferenceConverterTestObject> TestObject = NewObject<UReferenceConverterTestObject>();
+		UE::ToolsetRegistry::FJsonValueOrError Result =
+			UE::ToolsetRegistry::TestHelpers::ExecuteToolCallWithJson(
+				TestObject, TEXT("TestOptionalObject"), TEXT(R"({"testObject": {"refPath": ""}})"));
+		TestFalse("Should not have error", Result.HasError());
+	});
+
+	It("Rejects null for non-optional object properties", [this]()
+	{
+		TObjectPtr<UReferenceConverterTestObject> TestObject = NewObject<UReferenceConverterTestObject>();
+		UE::ToolsetRegistry::FJsonValueOrError Result =
+			UE::ToolsetRegistry::TestHelpers::ExecuteToolCallWithJson(
+				TestObject, TEXT("TestObjectParam"),
+				TEXT(R"({"testObject": {"refPath": ""}})"));
+
+		if (!TestTrue("Should have error", Result.HasError()))
+		{
+			return;
+		}
+		TestTrue("Error message", Result.GetError().Contains(TEXT("None is not valid value for property")));
+	});
+
+	It("Rejects invalid object path in object properties", [this]()
+	{
+		TObjectPtr<UReferenceConverterTestObject> TestObject = NewObject<UReferenceConverterTestObject>();
+		UE::ToolsetRegistry::FJsonValueOrError Result =
+			UE::ToolsetRegistry::TestHelpers::ExecuteToolCallWithJson(
+				TestObject, TEXT("TestObjectParam"),
+				TEXT(R"({"testObject": {"refPath": "/Invalid/Path/DoesNotExist"}})"));
+		if (!TestTrue("Should have error", Result.HasError()))
+		{
+			return;
+		}
+		TestTrue("Error contains validation message",
+				 Result.GetError().Contains(TEXT("is not valid Object for property")));
+	});
+
+	It("Rejects wrong object type in object properties", [this]()
+	{
+		TObjectPtr<UReferenceConverterTestObject> TestObject = NewObject<UReferenceConverterTestObject>();
+		UE::ToolsetRegistry::FJsonValueOrError Result =
+			UE::ToolsetRegistry::TestHelpers::ExecuteToolCallWithJson(
+				TestObject, TEXT("TestSpecObjectParam"),
+				FString::Printf(
+					TEXT(R"({"testObject": {"refPath": "%s"}})"), *UObject::StaticClass()->GetPathName()));
+		if (!TestTrue("Should have error", Result.HasError()))
+		{
+			return;
+		}
+		TestTrue("Error contains validation message",
+				 Result.GetError().Contains(TEXT("is not valid ReferenceConverterTestObject for property")));
+	});
+
+	It("Rejects invalid path for class properties", [this]()
+	{
+		TObjectPtr<UReferenceConverterTestObject> TestObject = NewObject<UReferenceConverterTestObject>();
+		UE::ToolsetRegistry::FJsonValueOrError Result =
+			UE::ToolsetRegistry::TestHelpers::ExecuteToolCallWithJson(
+				TestObject, TEXT("TestClassParam"),
+				TEXT(R"({"testClass": {"refPath": "/Invalid/Class/Path"}})"));
+		if (!TestTrue("Should have error", Result.HasError()))
+		{
+			return;
+		}
+		TestTrue("Error contains validation message",
+				 Result.GetError().Contains(TEXT("is not valid Class for property")));
+	});
+
+	It("Rejects non-class object for class properties", [this]()
+	{
+		TObjectPtr<UReferenceConverterTestObject> TestObj = NewObject<UReferenceConverterTestObject>();
+		TObjectPtr<UReferenceConverterTestObject> TestObject = NewObject<UReferenceConverterTestObject>();
+		UE::ToolsetRegistry::FJsonValueOrError Result =
+			UE::ToolsetRegistry::TestHelpers::ExecuteToolCallWithJson(
+				TestObject, TEXT("TestClassParam"),
+				FString::Printf(TEXT(R"({"testClass": {"refPath": "%s"}})"), *TestObj->GetPathName()));
+		if (!TestTrue("Should have error", Result.HasError()))
+		{
+			return;
+		}
+		TestTrue("Error contains validation message",
+				 Result.GetError().Contains(TEXT("is not valid Class for property")));
+	});
+
+	It("Rejects class not matching metaclass constraint in properties", [this]()
+	{
+		FString WrongClassPath = UObject::StaticClass()->GetPathName();
+		TObjectPtr<UReferenceConverterTestObject> TestObject = NewObject<UReferenceConverterTestObject>();
+		UE::ToolsetRegistry::FJsonValueOrError Result =
+			UE::ToolsetRegistry::TestHelpers::ExecuteToolCallWithJson(
+				TestObject, TEXT("TestSubClassParam"),
+				FString::Printf(TEXT(R"({"testClass": {"refPath": "%s"}})"), *WrongClassPath));
+		if (!TestTrue("Should have error", Result.HasError()))
+		{
+			return;
+		}
+		TestTrue("Error contains validation message",
+				 Result.GetError().Contains(TEXT("is not valid ReferenceConverterTestObject for property")));
+	});
+
+	It("Rejects wrong type for soft object reference", [this]()
+	{
+		FString InvalidJson = FString::Printf(
+			TEXT(R"({"testObject": {"refPath": "%s"}})"),
+			*UObject::StaticClass()->GetPathName());
+
+		TObjectPtr<UReferenceConverterTestObject> TestObject = NewObject<UReferenceConverterTestObject>();
+		UE::ToolsetRegistry::FJsonValueOrError Result =
+			UE::ToolsetRegistry::TestHelpers::ExecuteToolCallWithJson(
+				TestObject, TEXT("TestSoftObjectParam"), InvalidJson);
+
+		if (!TestTrue("Should have error", Result.HasError())) return;
+		TestTrue("Error message",
+			Result.GetError().Contains(TEXT("is not valid ReferenceConverterTestObject")));
+	});
+
+	It("Allows valid soft object reference", [this]()
+	{
+		UReferenceConverterTestObject* TestObj = NewObject<UReferenceConverterTestObject>();
+		FString ValidJson = FString::Printf(
+			TEXT(R"({"testObject": {"refPath": "%s"}})"),
+			*TestObj->GetPathName());
+
+		TObjectPtr<UReferenceConverterTestObject> TestObject = NewObject<UReferenceConverterTestObject>();
+		UE::ToolsetRegistry::FJsonValueOrError Result =
+			UE::ToolsetRegistry::TestHelpers::ExecuteToolCallWithJson(
+				TestObject, TEXT("TestSoftObjectParam"), ValidJson);
+
+		TestFalse("Should not have error", Result.HasError());
+	});
+
+	It("Accepts Blueprint asset path for FSoftClassPath parameter", [this]()
+	{
+		// Regression test: before the fix, FSoftClassPath validation called AssetData.GetClass()
+		// which returns UBlueprint for Blueprint assets. Since UBlueprint is not a child of UClass,
+		// it incorrectly raised a script error. The fix adds an explicit check for UBlueprint.
+
+		// Create a transient Blueprint and register it with the Asset Registry so that
+		// AssetData.IsValid() is true, which is required to exercise the validation branch.
+		UBlueprint* Blueprint = NewObject<UBlueprint>(
+			GetTransientPackage(), TEXT("TestBlueprintForSoftClassPathValidation"), RF_Transient);
+		Blueprint->AddToRoot();
+
+		FAssetRegistryModule& AssetRegistryModule =
+			FModuleManager::LoadModuleChecked<FAssetRegistryModule>("AssetRegistry");
+		IAssetRegistry& AssetRegistry = AssetRegistryModule.Get();
+		AssetRegistry.AssetCreated(Blueprint);
+
+		ON_SCOPE_EXIT
+		{
+			AssetRegistry.AssetDeleted(Blueprint);
+			Blueprint->RemoveFromRoot();
+		};
+
+		FString BlueprintPath = Blueprint->GetPathName();
+		FString Json = FString::Printf(
+			TEXT(R"({"testClassPath": {"refPath": "%s"}})"), *BlueprintPath);
+
+		TObjectPtr<UReferenceConverterTestObject> TestObject = NewObject<UReferenceConverterTestObject>();
+		UE::ToolsetRegistry::FJsonValueOrError Result =
+			UE::ToolsetRegistry::TestHelpers::ExecuteToolCallWithJson(
+				TestObject, TEXT("TestSoftClassPathParam"), Json);
+
+		TestFalse("Blueprint asset path should be accepted for FSoftClassPath", Result.HasError());
+	});
+
+	It("Rejects non-class for soft class reference", [this]()
+	{
+		UReferenceConverterTestObject* TestObj = NewObject<UReferenceConverterTestObject>();
+		FString InvalidJson = FString::Printf(
+			TEXT(R"({"testClass": {"refPath": "%s"}})"),
+			*TestObj->GetPathName());
+
+		TObjectPtr<UReferenceConverterTestObject> TestObject = NewObject<UReferenceConverterTestObject>();
+		UE::ToolsetRegistry::FJsonValueOrError Result =
+			UE::ToolsetRegistry::TestHelpers::ExecuteToolCallWithJson(
+				TestObject, TEXT("TestSoftClassParam"), InvalidJson);
+		if (!TestTrue("Should have error", Result.HasError())) return;
+		TestTrue("Error message",
+			Result.GetError().Contains(TEXT("is not valid Class for property")));
+	});
+
+	It("Rejects wrong class type for soft class reference", [this]()
+	{
+		FString InvalidJson = FString::Printf(
+			TEXT(R"({"testClass": {"refPath": "%s"}})"),
+			*UObject::StaticClass()->GetPathName());
+
+		TObjectPtr<UReferenceConverterTestObject> TestObject = NewObject<UReferenceConverterTestObject>();
+		UE::ToolsetRegistry::FJsonValueOrError Result =
+			UE::ToolsetRegistry::TestHelpers::ExecuteToolCallWithJson(
+				TestObject, TEXT("TestSoftClassParam"), InvalidJson);
+
+		if (!TestTrue("Should have error", Result.HasError())) return;
+		TestTrue("Error message",
+			Result.GetError().Contains(TEXT("is not valid ReferenceConverterTestObject")));
+	});
+
+	It("Allows valid soft class reference", [this]()
+	{
+		FString ValidJson = FString::Printf(
+			TEXT(R"({"testClass": {"refPath": "%s"}})"),
+			*UReferenceConverterTestObject::StaticClass()->GetPathName());
+
+		TObjectPtr<UReferenceConverterTestObject> TestObject = NewObject<UReferenceConverterTestObject>();
+		UE::ToolsetRegistry::FJsonValueOrError Result =
+			UE::ToolsetRegistry::TestHelpers::ExecuteToolCallWithJson(
+				TestObject, TEXT("TestSoftClassParam"), ValidJson);
+
+		TestFalse("Should not have error", Result.HasError());
+	});
+
+	It("Converts SoftPath property to reference schema", [this]()
+	{
+		UStruct* Struct = FReferenceConverterTest::StaticStruct();
+		FProperty* Prop = Struct->FindPropertyByName("TestObjectPath");
+		const TSharedPtr<FJsonObject> Schema = PropertyToJsonSchema(Prop);
+		TestObjectRefSchema(Schema, UObject::StaticClass(), false);
+	});
+
+	It("Converts SoftPath property to and from JSON", [this]()
+	{
+		FReferenceConverterTest TestIn;
+		TObjectPtr<UReferenceConverterTestObject> TestOther = NewObject<UReferenceConverterTestObject>();
+		TestIn.TestObjectPath = TestOther;
+
+		UStruct* Struct = FReferenceConverterTest::StaticStruct();
+		FProperty* Property = Struct->FindPropertyByName("TestObjectPath");
+
+		TSharedPtr<FJsonValue> OutJson = PropertyToJsonData(Property, &TestIn.TestObjectPath);
+		TestObjectRefData(OutJson->AsObject(), TestOther->GetPathName());
+
+		FReferenceConverterTest TestOut;
+		TestTrue("Output", JsonDataToProperty(OutJson, Property, &TestOut.TestObjectPath));
+		TestEqual("RefCheck", TestIn.TestObjectPath, TestOut.TestObjectPath);
+	});
+
+	It("Converts SoftPtr property to reference schema", [this]()
+	{
+		UStruct* Struct = FReferenceConverterTest::StaticStruct();
+		FProperty* Prop = Struct->FindPropertyByName("TestObjectSoft");
+		const TSharedPtr<FJsonObject> Schema = PropertyToJsonSchema(Prop);
+		TestObjectRefSchema(Schema, UReferenceConverterTestObject::StaticClass(), false);
+	});
+
+	It("Converts SoftPtr property to and from JSON", [this]()
+	{
+		FReferenceConverterTest TestIn;
+		TObjectPtr<UReferenceConverterTestObject> TestOther = NewObject<UReferenceConverterTestObject>();
+		TestIn.TestObjectSoft = TestOther;
+
+		UStruct* Struct = FReferenceConverterTest::StaticStruct();
+		FProperty* Property = Struct->FindPropertyByName("TestObjectSoft");
+
+		TSharedPtr<FJsonValue> OutJson = PropertyToJsonData(Property, &TestIn.TestObjectSoft);
+		TestObjectRefData(OutJson->AsObject(), TestOther->GetPathName());
+
+		FReferenceConverterTest TestOut;
+		TestTrue("Output", JsonDataToProperty(OutJson, Property, &TestOut.TestObjectSoft));
+		TestEqual("RefCheck", TestIn.TestObjectSoft, TestOut.TestObjectSoft);
+	});
+
+	It("Converts WeakPtr property to reference schema", [this]()
+	{
+		UStruct* Struct = FReferenceConverterTest::StaticStruct();
+		FProperty* Prop = Struct->FindPropertyByName("TestObjectWeak");
+		const TSharedPtr<FJsonObject> Schema = PropertyToJsonSchema(Prop);
+		TestObjectRefSchema(Schema, UReferenceConverterTestObject::StaticClass(), false);
+	});
+
+	It("Converts WeakPtr property to and from JSON", [this]()
+	{
+		FReferenceConverterTest TestIn;
+		TObjectPtr<UReferenceConverterTestObject> TestOther = NewObject<UReferenceConverterTestObject>();
+		TestIn.TestObjectWeak = TestOther;
+
+		UStruct* Struct = FReferenceConverterTest::StaticStruct();
+		FProperty* Property = Struct->FindPropertyByName("TestObjectWeak");
+
+		TSharedPtr<FJsonValue> OutJson = PropertyToJsonData(Property, &TestIn.TestObjectWeak);
+		TestObjectRefData(OutJson->AsObject(), TestOther->GetPathName());
+
+		FReferenceConverterTest TestOut;
+		TestTrue("Output", JsonDataToProperty(OutJson, Property, &TestOut.TestObjectWeak));
+		TestEqual("RefCheck", TestIn.TestObjectWeak, TestOut.TestObjectWeak);
+	});
+
+	It("Converts UClass* property to reference schema", [this]()
+	{
+		UStruct* Struct = FReferenceConverterTest::StaticStruct();
+		FProperty* Prop = Struct->FindPropertyByName("TestClass");
+		const TSharedPtr<FJsonObject> Schema = PropertyToJsonSchema(Prop);
+		TestObjectRefSchema(Schema, UReferenceConverterTestObject::StaticClass(), true);
+	});
+
+	It("Converts UClass* property to and from JSON", [this]()
+	{
+		FReferenceConverterTest TestIn;
+		TestIn.TestClass = UReferenceConverterTestObject::StaticClass();
+
+		UStruct* Struct = FReferenceConverterTest::StaticStruct();
+		FProperty* Property = Struct->FindPropertyByName("TestClass");
+
+		TSharedPtr<FJsonValue> OutJson = PropertyToJsonData(Property, &TestIn.TestClass);
+		TestObjectRefData(OutJson->AsObject(), UReferenceConverterTestObject::StaticClass()->GetPathName());
+
+		FReferenceConverterTest TestOut;
+		TestTrue("Output", JsonDataToProperty(OutJson, Property, &TestOut.TestClass));
+		TestEqual("RefCheck", TestIn.TestClass, TestOut.TestClass);
+	});
+
+	It("Converts SoftClassPath property to reference schema", [this]()
+	{
+		UStruct* Struct = FReferenceConverterTest::StaticStruct();
+		FProperty* Prop = Struct->FindPropertyByName("TestClassPath");
+		const TSharedPtr<FJsonObject> Schema = PropertyToJsonSchema(Prop);
+		TestObjectRefSchema(Schema, UObject::StaticClass(), true);
+	});
+
+	It("Converts SoftClassPath property to and from JSON", [this]()
+	{
+		FReferenceConverterTest TestIn;
+		TestIn.TestClassPath = UReferenceConverterTestObject::StaticClass();
+
+		UStruct* Struct = FReferenceConverterTest::StaticStruct();
+		FProperty* Property = Struct->FindPropertyByName("TestClassPath");
+
+		TSharedPtr<FJsonValue> OutJson = PropertyToJsonData(Property, &TestIn.TestClassPath);
+		TestObjectRefData(OutJson->AsObject(), UReferenceConverterTestObject::StaticClass()->GetPathName());
+
+		FReferenceConverterTest TestOut;
+		TestTrue("Output", JsonDataToProperty(OutJson, Property, &TestOut.TestClassPath));
+		TestEqual("RefCheck", TestIn.TestClassPath, TestOut.TestClassPath);
+	});
+
+	It("Imports from old string style path", [this]()
+	{
+		TObjectPtr<UReferenceConverterTestObject> TestOther = NewObject<UReferenceConverterTestObject>();
+		UStruct* Struct = FReferenceConverterTest::StaticStruct();
+		FProperty* Property = Struct->FindPropertyByName("TestObject");
+
+		FReferenceConverterTest TestOut;
+		TSharedPtr<FJsonValue> OutJson = MakeShared<TJsonValueString<TCHAR>>(TestOther->GetPathName());
+		TestTrue("Output", JsonDataToProperty(OutJson, Property, &TestOut.TestObject));
+		TestEqual("RefCheck", TestOut.TestObject, TestOther);
+	});
+
+	It("Can read null default for object reference", [this]()
+	{
+		UFunction* FunctionPtr = UReferenceConverterTestObject::StaticClass()->FindFunctionByName("TestDefault");
+		const TSharedPtr<FJsonObject> Schema = StructToJsonSchema(FunctionPtr);
+		if (!TestTrue("Schema", Schema.IsValid())) return;
+
+		const TSharedPtr<FJsonObject>* InputSchema = nullptr;
+		if (!(TestTrue("Input", Schema->TryGetObjectField(TEXT("inputSchema"), InputSchema) &&
+			InputSchema)))
+		{
+			return;
+		}
+
+		TestFalse("NoRequired", InputSchema->Get()->HasField(TEXT("required")));
+
+		const TSharedPtr<FJsonObject>* Properties = nullptr;
+		if (!(TestTrue("Properties", InputSchema->Get()->TryGetObjectField(TEXT("properties"), Properties) &&
+			Properties)))
+		{
+			return;
+		}
+		
+		const TSharedPtr<FJsonObject>* TestObject = nullptr;
+		if (!(TestTrue("testObject", Properties->Get()->TryGetObjectField(TEXT("testObject"), TestObject) &&
+			TestObject)))
+		{
+			return;
+		}
+		TestTrue("Default", TestObject->Get()->GetField(TEXT("default"), EJson::Null)->IsNull());
+	});
+}
+
+#endif
diff --git a/Engine/Plugins/Experimental/ToolsetRegistry/Source/ToolsetRegistry/Private/Tests/ReferenceConverterTest.h b/Engine/Plugins/Experimental/ToolsetRegistry/Source/ToolsetRegistry/Private/Tests/ReferenceConverterTest.h
new file mode 100644
--- /dev/null
+++ b/Engine/Plugins/Experimental/ToolsetRegistry/Source/ToolsetRegistry/Private/Tests/ReferenceConverterTest.h
@@ -0,0 +1,66 @@
+// Copyright Epic Games, Inc. All Rights Reserved.
+
+#pragma once
+
+#include "CoreMinimal.h"
+#include "Templates/SubclassOf.h"
+
+#include "ReferenceConverterTest.generated.h"
+
+
+UCLASS(BlueprintType)
+class UReferenceConverterTestObject : public UObject
+{
+	GENERATED_BODY()
+public:
+	UFUNCTION(BlueprintCallable, Category="ReferenceConverterTest")
+	void TestOptionalObject(UObject* TestObject = nullptr) {}
+
+	UFUNCTION(BlueprintCallable, Category="ReferenceConverterTest")
+	void TestObjectParam(UObject* TestObject) {}
+
+	UFUNCTION(BlueprintCallable, Category="ReferenceConverterTest")
+	void TestSpecObjectParam(UReferenceConverterTestObject* TestObject) {}
+
+	UFUNCTION(BlueprintCallable, Category="ReferenceConverterTest")
+	void TestClassParam(UClass* TestClass) {}
+
+	UFUNCTION(BlueprintCallable, Category="ReferenceConverterTest")
+	void TestSubClassParam(TSubclassOf<UReferenceConverterTestObject> TestClass) {}
+
+	UFUNCTION(BlueprintCallable, Category="ReferenceConverterTest")
+	void TestSoftObjectParam(TSoftObjectPtr<UReferenceConverterTestObject> TestObject) {}
+
+	UFUNCTION(BlueprintCallable, Category="ReferenceConverterTest")
+	void TestSoftClassParam(TSoftClassPtr<UReferenceConverterTestObject> TestClass) {}
+
+	UFUNCTION(BlueprintCallable, Category="ReferenceConverterTest")
+	void TestSoftClassPathParam(FSoftClassPath TestClassPath) {}
+
+	UFUNCTION(BlueprintCallable, Category="ReferenceConverterTest")
+	void TestDefault(UObject* TestObject = nullptr) {}
+};
+
+USTRUCT(BlueprintType)
+struct FReferenceConverterTest
+{
+	GENERATED_BODY()
+public:
+	UPROPERTY()
+	TObjectPtr<UReferenceConverterTestObject> TestObject;
+
+	UPROPERTY()
+	FSoftObjectPath TestObjectPath;
+
+	UPROPERTY()
+	TSoftObjectPtr<UReferenceConverterTestObject> TestObjectSoft;
+
+	UPROPERTY()
+	TWeakObjectPtr<UReferenceConverterTestObject> TestObjectWeak;
+
+	UPROPERTY()
+	TSubclassOf<UReferenceConverterTestObject> TestClass;
+
+	UPROPERTY()
+	FSoftClassPath TestClassPath;
+};
diff --git a/Engine/Plugins/Experimental/ToolsetRegistry/Source/ToolsetRegistry/Private/Tests/ToolCallScriptExceptionHandlerTest.cpp b/Engine/Plugins/Experimental/ToolsetRegistry/Source/ToolsetRegistry/Private/Tests/ToolCallScriptExceptionHandlerTest.cpp
--- a/Engine/Plugins/Experimental/ToolsetRegistry/Source/ToolsetRegistry/Private/Tests/ToolCallScriptExceptionHandlerTest.cpp
+++ b/Engine/Plugins/Experimental/ToolsetRegistry/Source/ToolsetRegistry/Private/Tests/ToolCallScriptExceptionHandlerTest.cpp
@@ -1,9 +1,9 @@
 // Copyright Epic Games, Inc. All Rights Reserved.
 
+#include "ToolCallScriptExceptionHandlerTest.h"
 
 #include "Misc/AutomationTest.h"
 
-#include "Tests/ToolsetJsonTest.h"
 #include "Tests/ToolCallTestHelpers.h"
 #include "Tests/ToolsetRegistryTestFlags.h"
 #include "ToolsetRegistry/JsonValueOrError.h"
@@ -25,7 +25,8 @@ void FToolCallScriptExceptionHandlerSpec::Define()
 {
 	It("Accumulates multiple parameter errors with newlines", [this]()
 	{
-		TObjectPtr<UToolsetJsonTestObject> TestObject = NewObject<UToolsetJsonTestObject>();
+		TObjectPtr<UToolCallScriptExceptionHandlerTestObject> TestObject =
+			NewObject<UToolCallScriptExceptionHandlerTestObject>();
 
 		// Pass two invalid references to trigger multiple RaiseScriptError calls
 		FString InvalidJson = FString::Printf(
@@ -42,9 +43,9 @@ void FToolCallScriptExceptionHandlerSpec::Define()
 
 		// Verify both errors are present
 		TestTrue("Contains first param type validation error",
-				 ErrorMessage.Contains(TEXT("is not valid ToolsetJsonTestObject for property 'FirstObject'")));
+				 ErrorMessage.Contains(TEXT("is not valid ToolCallScriptExceptionHandlerTestObject for property 'FirstObject'")));
 		TestTrue("Contains second param type validation error",
-				 ErrorMessage.Contains(TEXT("is not valid ToolsetJsonTestObject for property 'SecondObject'")));
+				 ErrorMessage.Contains(TEXT("is not valid ToolCallScriptExceptionHandlerTestObject for property 'SecondObject'")));
 
 		// Verify errors are accumulated with newlines
 		TestTrue("Contains newline separator",
@@ -57,7 +58,8 @@ void FToolCallScriptExceptionHandlerSpec::Define()
 
 	It("Captures single parameter error correctly", [this]()
 	{
-		TObjectPtr<UToolsetJsonTestObject> TestObject = NewObject<UToolsetJsonTestObject>();
+		TObjectPtr<UToolCallScriptExceptionHandlerTestObject> TestObject =
+			NewObject<UToolCallScriptExceptionHandlerTestObject>();
 
 		// Pass one invalid reference
 		FString InvalidJson = FString::Printf(
@@ -73,16 +75,19 @@ void FToolCallScriptExceptionHandlerSpec::Define()
 
 		// Verify error message format
 		TestTrue("Contains type validation error",
-				 ErrorMessage.Contains(TEXT("is not valid ToolsetJsonTestObject")));
+				 ErrorMessage.Contains(TEXT("is not valid ToolCallScriptExceptionHandlerTestObject")));
 		TestTrue("Error has parameter prefix",
 				 ErrorMessage.Contains(TEXT("Parameter error")));
 	});
 
 	It("Returns no error when all parameters are valid", [this]()
 	{
-		TObjectPtr<UToolsetJsonTestObject> TestObject = NewObject<UToolsetJsonTestObject>();
-		TObjectPtr<UToolsetJsonTestObject> ValidObject1 = NewObject<UToolsetJsonTestObject>();
-		TObjectPtr<UToolsetJsonTestObject> ValidObject2 = NewObject<UToolsetJsonTestObject>();
+		TObjectPtr<UToolCallScriptExceptionHandlerTestObject> TestObject =
+			NewObject<UToolCallScriptExceptionHandlerTestObject>();
+		TObjectPtr<UToolCallScriptExceptionHandlerTestObject> ValidObject1 =
+			NewObject<UToolCallScriptExceptionHandlerTestObject>();
+		TObjectPtr<UToolCallScriptExceptionHandlerTestObject> ValidObject2 =
+			NewObject<UToolCallScriptExceptionHandlerTestObject>();
 
 		// Pass valid references
 		FString ValidJson = FString::Printf(
@@ -99,7 +104,8 @@ void FToolCallScriptExceptionHandlerSpec::Define()
 
 	It("Returns a fallback message when an empty error message is raised", [this]()
 	{
-		TObjectPtr<UToolsetJsonTestObject> TestObject = NewObject<UToolsetJsonTestObject>();
+		TObjectPtr<UToolCallScriptExceptionHandlerTestObject> TestObject =
+			NewObject<UToolCallScriptExceptionHandlerTestObject>();
 
 		UFunction* Function = TestObject->GetClass()->FindFunctionByName(TEXT("TestEmptyError"));
 		if (!TestNotNull("Function exists", Function)) return;
diff --git a/Engine/Plugins/Experimental/ToolsetRegistry/Source/ToolsetRegistry/Private/Tests/ToolCallScriptExceptionHandlerTest.h b/Engine/Plugins/Experimental/ToolsetRegistry/Source/ToolsetRegistry/Private/Tests/ToolCallScriptExceptionHandlerTest.h
new file mode 100644
--- /dev/null
+++ b/Engine/Plugins/Experimental/ToolsetRegistry/Source/ToolsetRegistry/Private/Tests/ToolCallScriptExceptionHandlerTest.h
@@ -0,0 +1,30 @@
+// Copyright Epic Games, Inc. All Rights Reserved.
+
+#pragma once
+
+#include "CoreMinimal.h"
+#include "Kismet/KismetSystemLibrary.h"
+
+#include "ToolCallScriptExceptionHandlerTest.generated.h"
+
+
+UCLASS(BlueprintType)
+class UToolCallScriptExceptionHandlerTestObject : public UObject
+{
+	GENERATED_BODY()
+public:
+	UFUNCTION(BlueprintCallable, Category="ToolCallScriptExceptionHandlerTest")
+	void TestSpecObjectParam(UToolCallScriptExceptionHandlerTestObject* TestObject) {}
+
+	UFUNCTION(BlueprintCallable, Category="ToolCallScriptExceptionHandlerTest")
+	void TestMultipleRefs(
+		UToolCallScriptExceptionHandlerTestObject* FirstObject,
+		UToolCallScriptExceptionHandlerTestObject* SecondObject) {}
+
+	UFUNCTION(BlueprintCallable, Category="ToolCallScriptExceptionHandlerTest")
+	void TestEmptyError()
+	{
+		// Raise an explicitly empty error
+		UKismetSystemLibrary::RaiseScriptError(FString());
+	}
+};
diff --git a/Engine/Plugins/Experimental/ToolsetRegistry/Source/ToolsetRegistry/Private/Tests/ToolsetJsonTest.cpp b/Engine/Plugins/Experimental/ToolsetRegistry/Source/ToolsetRegistry/Private/Tests/ToolsetJsonTest.cpp
--- a/Engine/Plugins/Experimental/ToolsetRegistry/Source/ToolsetRegistry/Private/Tests/ToolsetJsonTest.cpp
+++ b/Engine/Plugins/Experimental/ToolsetRegistry/Source/ToolsetRegistry/Private/Tests/ToolsetJsonTest.cpp
@@ -2,210 +2,20 @@
 
 #include "ToolsetJsonTest.h"
 
-#include "AssetRegistry/AssetData.h"
-#include "AssetRegistry/AssetRegistryModule.h"
 #include "CoreMinimal.h"
-#include "Containers/Array.h"
 #include "Containers/ArrayView.h"
-#include "Engine/Blueprint.h"
-#include "JsonObjectConverter.h"
 #include "Misc/AutomationTest.h"
 
 #include "ToolsetRegistry/JsonConversion.h"
 #include "ToolsetRegistry/JsonSchema.h"
 #include "ToolsetRegistry/ToolsetJson.h"
-#include "ToolsetRegistry/ToolCallAsyncResult.h"
 #include "ToolsetRegistry/ToolCallAsyncResultString.h"
 #include "Tests/ToolCallAsyncResultTest.h"
-#include "Tests/ToolCallTestHelpers.h"
 
 #if WITH_DEV_AUTOMATION_TESTS
 
 namespace
 {
-	void TestObjectRefSchema(
-		FAutomationTestBase& Test, const TSharedPtr<FJsonObject>& Schema,
-		UClass* Class, bool IsClass)
-	{
-		Test.TestTrue("Schema", Schema.IsValid());
-		if (!Schema.IsValid())
-		{
-			return;
-		}
-
-		FString Type;
-		Test.TestTrue("Type", Schema->TryGetStringField(TEXT("type"), Type));
-		Test.TestEqual("Type", Type, FString(TEXT("object")));
-
-		if (Class)
-		{
-			FString Title;
-			Test.TestTrue("Title", Schema->TryGetStringField(TEXT("title"), Title));
-			FString ClassPrefix(IsClass ? TEXT("Class@") : TEXT(""));
-			Test.TestEqual("Title", Title, ClassPrefix + Class->GetPathName());
-		}
-	
-		const TSharedPtr<FJsonObject>* Properties = nullptr;
-		Test.TestTrue("Properties", Schema->TryGetObjectField(TEXT("properties"), Properties));
-		if (Properties)
-		{
-			const TSharedPtr<FJsonObject>* RefPath = nullptr;
-			Test.TestTrue(
-				"RefPath", Properties->Get()->TryGetObjectField(TEXT("refPath"), RefPath));
-			if (RefPath)
-			{
-				Test.TestTrue("RefPath", RefPath->Get()->TryGetStringField(TEXT("type"), Type));
-				Test.TestEqual("RefPath", Type, FString(TEXT("string")));
-			}
-		}
-	}
-
-	void TestObjectRefData(
-		FAutomationTestBase& Test, const TSharedPtr<FJsonObject>& Schema, const FString& Expected)
-	{
-		Test.TestTrue("Schema", Schema.IsValid());
-		if (!Schema.IsValid())
-		{
-			return;
-		}
-	
-		FString RefPath;
-		Test.TestTrue("Unpack", Schema->TryGetStringField(TEXT("refPath"), RefPath));
-		Test.TestEqual("RefPath", RefPath, Expected);
-	}
-
-	void TestObjectRefDataInStruct(
-		FAutomationTestBase& Test, const TSharedPtr<FJsonObject>& Schema,
-		const FString& PropertyName, const FString& Expected)
-	{
-		Test.TestTrue("Schema", Schema.IsValid());
-		if (!Schema.IsValid())
-		{
-			return;
-		}
-
-		const TSharedPtr<FJsonObject>* Property = nullptr;
-		Test.TestTrue("Unpack", Schema->TryGetObjectField(PropertyName, Property));
-		if (Property)
-		{
-			TestObjectRefData(Test, *Property, Expected);
-		}
-	}
-
-	void TestColorSchema(
-		FAutomationTestBase& Test, const TSharedPtr<FJsonObject>& Schema, bool HasMax)
-	{
-		Test.TestTrue("Schema", Schema.IsValid());
-		if (!Schema.IsValid())
-		{
-			return;
-		}
-
-		FString Type;
-		Test.TestTrue("Type", Schema->TryGetStringField(TEXT("type"), Type));
-		Test.TestEqual("Type", Type, FString(TEXT("object")));
-
-		FString Title;
-		Test.TestTrue("Title", Schema->TryGetStringField(TEXT("title"), Title));
-		Test.TestEqual("Title", Title, FString(TEXT("LinearColor")));
-
-		Test.TestFalse("Description", Schema->HasField(TEXT("description")));
-
-		const TSharedPtr<FJsonObject>* Properties = nullptr;
-		Test.TestTrue("Properties", Schema->TryGetObjectField(TEXT("properties"), Properties));
-		if (Properties)
-		{
-			TArray<FString> PropertyNames({ "r", "g", "b", "a" });
-			for (const auto& Name : PropertyNames)
-			{
-				const TSharedPtr<FJsonObject>* Property = nullptr;
-				Test.TestTrue(
-					"Property", Properties->Get()->TryGetObjectField(Name, Property));
-				if (Property)
-				{
-					Test.TestEqual("Property", Property->Get()->GetStringField(
-						TEXT("type")), FString(TEXT("number")));
-					Test.TestEqual("Property", Property->Get()->GetNumberField(
-						TEXT("minimum")), 0.0);
-					if (HasMax)
-					{
-						Test.TestEqual("Property", Property->Get()->GetNumberField(
-							TEXT("maximum")), 1.0);
-					}
-				}
-			}
-		}
-	}
-
-	void TestColorData(
-		FAutomationTestBase& Test, const TSharedPtr<FJsonObject>& Schema, const FColor& Expected)
-	{
-		Test.TestTrue("Schema", Schema.IsValid());
-		if (!Schema.IsValid())
-		{
-			return;
-		}
-
-		const FLinearColor LinearExpected = Expected.ReinterpretAsLinear();
-
-		float R;
-		Test.TestTrue("Color", Schema->TryGetNumberField(TEXT("r"), R));
-		Test.TestEqual("Color", R, LinearExpected.R, 2e-3f);
-
-		float G;
-		Test.TestTrue("Color", Schema->TryGetNumberField(TEXT("g"), G));
-		Test.TestEqual("Color", G, LinearExpected.G, 2e-3f);
-
-		float B;
-		Test.TestTrue("Color", Schema->TryGetNumberField(TEXT("b"), B));
-		Test.TestEqual("Color", B, LinearExpected.B, 2e-3f);
-
-		float A;
-		Test.TestTrue("Color", Schema->TryGetNumberField(TEXT("a"), A));
-		Test.TestEqual("Color", A, LinearExpected.A, 2e-3f);
-	}
-
-	void TestColorDataInStruct(
-		FAutomationTestBase& Test, const TSharedPtr<FJsonObject>& Schema,
-		const FString& PropertyName, const FColor& Expected)
-	{
-		Test.TestTrue("Schema", Schema.IsValid());
-		if (!Schema.IsValid())
-		{
-			return;
-		}
-
-		const TSharedPtr<FJsonObject>* Property = nullptr;
-		Test.TestTrue("Unpack", Schema->TryGetObjectField(PropertyName, Property));
-		if (Property)
-		{
-			TestColorData(Test, *Property, Expected);
-		}
-	}
-
-	void TestTransformSchema(FAutomationTestBase& Test, const TSharedPtr<FJsonObject>& Schema)
-	{
-		Test.TestTrue("Schema", Schema.IsValid());
-		if (!Schema.IsValid()) return;
-
-		FString Type;
-		Test.TestTrue("Type", Schema->TryGetStringField(TEXT("type"), Type));
-		Test.TestEqual("Type", Type, FString(TEXT("object")));
-
-		Test.TestFalse("No required", Schema->HasField(TEXT("required")));
-
-		const TSharedPtr<FJsonObject>* Properties = nullptr;
-		Test.TestTrue("Properties", Schema->TryGetObjectField(TEXT("properties"), Properties));
-		if (Properties)
-		{
-			for (const FString& Field : TArray<FString>{ TEXT("location"), TEXT("rotation"), TEXT("scale") })
-			{
-				const TSharedPtr<FJsonObject>* SubSchema = nullptr;
-				Test.TestTrue(*Field, Properties->Get()->TryGetObjectField(Field, SubSchema));
-			}
-		}
-	}
-
 	TSharedPtr<FJsonObject> TestHasObjectField(
 		FAutomationTestBase& Test, const TSharedPtr<FJsonObject>& Schema,
 		const FString& RootObjectName, const TArrayView<FString>& FieldPath)
@@ -244,726 +54,28 @@ void FToolsetJsonSpec::Define()
 	using namespace UE::ToolsetRegistry::Internal;
 	using namespace UE::ToolsetRegistry::Internal::ToolsetJson;
 
-	It("Converts UObject* property to reference schema", [this]()
-	{
-		UStruct* Struct = FToolsetJsonTest::StaticStruct();
-		FProperty* Prop = Struct->FindPropertyByName("TestObject");
-		const TSharedPtr<FJsonObject> Schema = PropertyToJsonSchema(Prop);
-		TestObjectRefSchema(*this, Schema, UToolsetJsonTestObject::StaticClass(), false);
-	});
-
-	It("Converts UObject* property to and from JSON", [this]()
-	{
-		FToolsetJsonTest TestIn;
-		TObjectPtr<UToolsetJsonTestObject> TestOther = NewObject<UToolsetJsonTestObject>();
-		TestIn.TestObject = TestOther;
-
-		UStruct* Struct = FToolsetJsonTest::StaticStruct();
-		FProperty* Property = Struct->FindPropertyByName("TestObject");
-
-		TSharedPtr<FJsonValue> OutJson = PropertyToJsonData(Property, &TestIn.TestObject);
-		TestObjectRefData(*this, OutJson->AsObject(), TestOther->GetPathName());
-
-		FToolsetJsonTest TestOut;
-		TestTrue("Output", JsonDataToProperty(OutJson, Property, &TestOut.TestObject));
-		TestEqual("RefCheck", TestIn.TestObject, TestOut.TestObject);
-	});
-
-	It("Ensures refs can be set from all forms of null", [this]()
-	{
-		FToolsetJsonTest Test;
-		TObjectPtr<UToolsetJsonTestObject> TestOther = NewObject<UToolsetJsonTestObject>();
-		UStruct* Struct = FToolsetJsonTest::StaticStruct();
-		FProperty* Property = Struct->FindPropertyByName("TestObject");
-		TSharedPtr<FJsonObject> JsonRef = MakeShared<FJsonObject>();
-
-		// Null Json type.
-		Test.TestObject = TestOther;
-		TestEqual("RefCheck", Test.TestObject, TestOther);
-		JsonRef->SetField(FString(TEXT("refPath")), MakeShared<FJsonValueNull>());
-		TestTrue("Output", JsonDataToProperty(
-			MakeShared<FJsonValueObject>(JsonRef), Property, &Test.TestObject));
-		TestEqual("RefCheck", Test.TestObject, TObjectPtr<UToolsetJsonTestObject>());
-
-		// None is the UE standard.
-		Test.TestObject = TestOther;
-		TestEqual("RefCheck", Test.TestObject, TestOther);
-		JsonRef->SetStringField(FString(TEXT("refPath")), FString(TEXT("None")));
-		TestTrue("Output", JsonDataToProperty(
-			MakeShared<FJsonValueObject>(JsonRef), Property, &Test.TestObject));
-		TestEqual("RefCheck", Test.TestObject, TObjectPtr<UToolsetJsonTestObject>());
-
-		// Empty string also means null.
-		Test.TestObject = TestOther;
-		TestEqual("RefCheck", Test.TestObject, TestOther);
-		JsonRef->SetStringField(FString(TEXT("refPath")), FString(TEXT("")));
-		TestTrue("Output", JsonDataToProperty(
-			MakeShared<FJsonValueObject>(JsonRef), Property, &Test.TestObject));
-		TestEqual("RefCheck", Test.TestObject, TObjectPtr<UToolsetJsonTestObject>());
-
-		// And we also support null just in case.
-		Test.TestObject = TestOther;
-		TestEqual("RefCheck", Test.TestObject, TestOther);
-		JsonRef->SetStringField(FString(TEXT("refPath")), FString(TEXT("null")));
-		TestTrue("Output", JsonDataToProperty(
-			MakeShared<FJsonValueObject>(JsonRef), Property, &Test.TestObject));
-		TestEqual("RefCheck", Test.TestObject, TObjectPtr<UToolsetJsonTestObject>());
-	});
-
-	It("Accepts null for optional parameter", [this]()
-	{
-		TObjectPtr<UToolsetJsonTestObject> TestObject = NewObject<UToolsetJsonTestObject>();
-		UE::ToolsetRegistry::FJsonValueOrError Result =
-			UE::ToolsetRegistry::TestHelpers::ExecuteToolCallWithJson(
-				TestObject, TEXT("TestOptionalObject"), TEXT(R"({"testObject": {"refPath": ""}})"));
-		TestFalse("Should not have error", Result.HasError());
-	});
-
-	It("Rejects null for non-optional object properties", [this]()
-	{
-		TObjectPtr<UToolsetJsonTestObject> TestObject = NewObject<UToolsetJsonTestObject>();
-		UE::ToolsetRegistry::FJsonValueOrError Result =
-			UE::ToolsetRegistry::TestHelpers::ExecuteToolCallWithJson(
-				TestObject, TEXT("TestObjectParam"),
-				TEXT(R"({"testObject": {"refPath": ""}})"));
-
-		if (!TestTrue("Should have error", Result.HasError()))
-		{
-			return;
-		}
-		TestTrue("Error message", Result.GetError().Contains(TEXT("None is not valid value for property")));
-	});
-
-	It("Rejects invalid object path in object properties", [this]()
-	{
-		TObjectPtr<UToolsetJsonTestObject> TestObject = NewObject<UToolsetJsonTestObject>();
-		UE::ToolsetRegistry::FJsonValueOrError Result =
-			UE::ToolsetRegistry::TestHelpers::ExecuteToolCallWithJson(
-				TestObject, TEXT("TestObjectParam"),
-				TEXT(R"({"testObject": {"refPath": "/Invalid/Path/DoesNotExist"}})"));
-		if (!TestTrue("Should have error", Result.HasError()))
-		{
-			return;
-		}
-		TestTrue("Error contains validation message",
-				 Result.GetError().Contains(TEXT("is not valid Object for property")));
-	});
-
-	It("Rejects wrong object type in object properties", [this]()
-	{
-		TObjectPtr<UToolsetJsonTestObject> TestObject = NewObject<UToolsetJsonTestObject>();
-		UE::ToolsetRegistry::FJsonValueOrError Result =
-			UE::ToolsetRegistry::TestHelpers::ExecuteToolCallWithJson(
-				TestObject, TEXT("TestSpecObjectParam"),
-				FString::Printf(
-					TEXT(R"({"testObject": {"refPath": "%s"}})"), *UObject::StaticClass()->GetPathName()));
-		if (!TestTrue("Should have error", Result.HasError()))
-		{
-			return;
-		}
-		TestTrue("Error contains validation message",
-				 Result.GetError().Contains(TEXT("is not valid ToolsetJsonTestObject for property")));
-	});
-
-	It("Rejects invalid path for class properties", [this]()
-	{
-		TObjectPtr<UToolsetJsonTestObject> TestObject = NewObject<UToolsetJsonTestObject>();
-		UE::ToolsetRegistry::FJsonValueOrError Result =
-			UE::ToolsetRegistry::TestHelpers::ExecuteToolCallWithJson(
-				TestObject, TEXT("TestClassParam"),
-				TEXT(R"({"testClass": {"refPath": "/Invalid/Class/Path"}})"));
-		if (!TestTrue("Should have error", Result.HasError()))
-		{
-			return;
-		}
-		TestTrue("Error contains validation message",
-				 Result.GetError().Contains(TEXT("is not valid Class for property")));
-	});
-
-	It("Rejects non-class object for class properties", [this]()
-	{
-		TObjectPtr<UToolsetJsonTestObject> TestObj = NewObject<UToolsetJsonTestObject>();
-		TObjectPtr<UToolsetJsonTestObject> TestObject = NewObject<UToolsetJsonTestObject>();
-		UE::ToolsetRegistry::FJsonValueOrError Result =
-			UE::ToolsetRegistry::TestHelpers::ExecuteToolCallWithJson(
-				TestObject, TEXT("TestClassParam"),
-				FString::Printf(TEXT(R"({"testClass": {"refPath": "%s"}})"), *TestObj->GetPathName()));
-		if (!TestTrue("Should have error", Result.HasError()))
-		{
-			return;
-		}
-		TestTrue("Error contains validation message",
-				 Result.GetError().Contains(TEXT("is not valid Class for property")));
-	});
-
-	It("Rejects class not matching metaclass constraint in properties", [this]()
-	{
-		FString WrongClassPath = UObject::StaticClass()->GetPathName();
-		TObjectPtr<UToolsetJsonTestObject> TestObject = NewObject<UToolsetJsonTestObject>();
-		UE::ToolsetRegistry::FJsonValueOrError Result =
-			UE::ToolsetRegistry::TestHelpers::ExecuteToolCallWithJson(
-				TestObject, TEXT("TestSubClassParam"),
-				FString::Printf(TEXT(R"({"testClass": {"refPath": "%s"}})"), *WrongClassPath));
-		if (!TestTrue("Should have error", Result.HasError()))
-		{
-			return;
-		}
-		TestTrue("Error contains validation message",
-				 Result.GetError().Contains(TEXT("is not valid ToolsetJsonTestObject for property")));
-	});
-
-	It("Rejects wrong type for soft object reference", [this]()
-	{
-		FString InvalidJson = FString::Printf(
-			TEXT(R"({"testObject": {"refPath": "%s"}})"),
-			*UObject::StaticClass()->GetPathName());
-
-		TObjectPtr<UToolsetJsonTestObject> TestObject = NewObject<UToolsetJsonTestObject>();
-		UE::ToolsetRegistry::FJsonValueOrError Result =
-			UE::ToolsetRegistry::TestHelpers::ExecuteToolCallWithJson(
-				TestObject, TEXT("TestSoftObjectParam"), InvalidJson);
-
-		if (!TestTrue("Should have error", Result.HasError())) return;
-		TestTrue("Error message",
-			Result.GetError().Contains(TEXT("is not valid ToolsetJsonTestObject")));
-	});
-
-	It("Allows valid soft object reference", [this]()
-	{
-		// Create a valid test object asset
-		UToolsetJsonTestObject* TestObj = NewObject<UToolsetJsonTestObject>();
-		FString ValidJson = FString::Printf(
-			TEXT(R"({"testObject": {"refPath": "%s"}})"),
-			*TestObj->GetPathName());
-
-		TObjectPtr<UToolsetJsonTestObject> TestObject = NewObject<UToolsetJsonTestObject>();
-		UE::ToolsetRegistry::FJsonValueOrError Result =
-			UE::ToolsetRegistry::TestHelpers::ExecuteToolCallWithJson(
-				TestObject, TEXT("TestSoftObjectParam"), ValidJson);
-
-		TestFalse("Should not have error", Result.HasError());
-	});
-
-	It("Accepts Blueprint asset path for FSoftClassPath parameter", [this]()
-	{
-		// Regression test: before the fix, FSoftClassPath validation called AssetData.GetClass()
-		// which returns UBlueprint for Blueprint assets. Since UBlueprint is not a child of UClass,
-		// it incorrectly raised a script error. The fix adds an explicit check for UBlueprint.
-
-		// Create a transient Blueprint and register it with the Asset Registry so that
-		// AssetData.IsValid() is true, which is required to exercise the validation branch.
-		UBlueprint* Blueprint = NewObject<UBlueprint>(
-			GetTransientPackage(), TEXT("TestBlueprintForSoftClassPathValidation"), RF_Transient);
-		Blueprint->AddToRoot();
-
-		FAssetRegistryModule& AssetRegistryModule =
-			FModuleManager::LoadModuleChecked<FAssetRegistryModule>("AssetRegistry");
-		IAssetRegistry& AssetRegistry = AssetRegistryModule.Get();
-		AssetRegistry.AssetCreated(Blueprint);
-
-		ON_SCOPE_EXIT
-		{
-			AssetRegistry.AssetDeleted(Blueprint);
-			Blueprint->RemoveFromRoot();
-		};
-
-		FString BlueprintPath = Blueprint->GetPathName();
-		FString Json = FString::Printf(
-			TEXT(R"({"testClassPath": {"refPath": "%s"}})"), *BlueprintPath);
-
-		TObjectPtr<UToolsetJsonTestObject> TestObject = NewObject<UToolsetJsonTestObject>();
-		UE::ToolsetRegistry::FJsonValueOrError Result =
-			UE::ToolsetRegistry::TestHelpers::ExecuteToolCallWithJson(
-				TestObject, TEXT("TestSoftClassPathParam"), Json);
-
-		TestFalse("Blueprint asset path should be accepted for FSoftClassPath", Result.HasError());
-	});
-
-	It("Rejects non-class for soft class reference", [this]()
-	{
-		// Try to pass a regular object instance where a class is expected
-		UToolsetJsonTestObject* TestObj = NewObject<UToolsetJsonTestObject>();
-		FString InvalidJson = FString::Printf(
-			TEXT(R"({"testClass": {"refPath": "%s"}})"),
-			*TestObj->GetPathName());
-
-		TObjectPtr<UToolsetJsonTestObject> TestObject = NewObject<UToolsetJsonTestObject>();
-		UE::ToolsetRegistry::FJsonValueOrError Result =
-			UE::ToolsetRegistry::TestHelpers::ExecuteToolCallWithJson(
-				TestObject, TEXT("TestSoftClassParam"), InvalidJson);
-		if (!TestTrue("Should have error", Result.HasError())) return;
-		TestTrue("Error message",
-			Result.GetError().Contains(TEXT("is not valid Class for property")));
-	});
-
-	It("Rejects wrong class type for soft class reference", [this]()
-	{
-		// Try to pass UObject class where UToolsetJsonTestObject-derived class is expected
-		FString InvalidJson = FString::Printf(
-			TEXT(R"({"testClass": {"refPath": "%s"}})"),
-			*UObject::StaticClass()->GetPathName());
-
-		TObjectPtr<UToolsetJsonTestObject> TestObject = NewObject<UToolsetJsonTestObject>();
-		UE::ToolsetRegistry::FJsonValueOrError Result =
-			UE::ToolsetRegistry::TestHelpers::ExecuteToolCallWithJson(
-				TestObject, TEXT("TestSoftClassParam"), InvalidJson);
-
-		if (!TestTrue("Should have error", Result.HasError())) return;
-		TestTrue("Error message",
-			Result.GetError().Contains(TEXT("is not valid ToolsetJsonTestObject")));
-	});
-
-	It("Allows valid soft class reference", [this]()
-	{
-		// Pass a valid class reference
-		FString ValidJson = FString::Printf(
-			TEXT(R"({"testClass": {"refPath": "%s"}})"),
-			*UToolsetJsonTestObject::StaticClass()->GetPathName());
-
-		TObjectPtr<UToolsetJsonTestObject> TestObject = NewObject<UToolsetJsonTestObject>();
-		UE::ToolsetRegistry::FJsonValueOrError Result =
-			UE::ToolsetRegistry::TestHelpers::ExecuteToolCallWithJson(
-				TestObject, TEXT("TestSoftClassParam"), ValidJson);
-
-		TestFalse("Should not have error", Result.HasError());
-	});
-
-	It("Converts SoftPath property to reference schema", [this]()
-	{
-		UStruct* Struct = FToolsetJsonTest::StaticStruct();
-		FProperty* Prop = Struct->FindPropertyByName("TestObjectPath");
-		const TSharedPtr<FJsonObject> Schema = PropertyToJsonSchema(Prop);
-		TestObjectRefSchema(*this, Schema, UObject::StaticClass(), false);
-	});
-
-	It("Converts SoftPath property to and from JSON", [this]()
-	{
-		FToolsetJsonTest TestIn;
-		TObjectPtr<UToolsetJsonTestObject> TestOther = NewObject<UToolsetJsonTestObject>();
-		TestIn.TestObjectPath = TestOther;
-
-		UStruct* Struct = FToolsetJsonTest::StaticStruct();
-		FProperty* Property = Struct->FindPropertyByName("TestObjectPath");
-
-		TSharedPtr<FJsonValue> OutJson = PropertyToJsonData(Property, &TestIn.TestObjectPath);
-		TestObjectRefData(*this, OutJson->AsObject(), TestOther->GetPathName());
-
-		FToolsetJsonTest TestOut;
-		TestTrue("Output", JsonDataToProperty(OutJson, Property, &TestOut.TestObjectPath));
-		TestEqual("RefCheck", TestIn.TestObjectPath, TestOut.TestObjectPath);
-	});
-
-	It("Converts SoftPtr property to reference schema", [this]()
-	{
-		UStruct* Struct = FToolsetJsonTest::StaticStruct();
-		FProperty* Prop = Struct->FindPropertyByName("TestObjectSoft");
-		const TSharedPtr<FJsonObject> Schema = PropertyToJsonSchema(Prop);
-		TestObjectRefSchema(*this, Schema, UToolsetJsonTestObject::StaticClass(), false);
-	});
-
-	It("Converts SoftPtr property to and from JSON", [this]()
-	{
-		FToolsetJsonTest TestIn;
-		TObjectPtr<UToolsetJsonTestObject> TestOther = NewObject<UToolsetJsonTestObject>();
-		TestIn.TestObjectSoft = TestOther;
-
-		UStruct* Struct = FToolsetJsonTest::StaticStruct();
-		FProperty* Property = Struct->FindPropertyByName("TestObjectSoft");
-
-		TSharedPtr<FJsonValue> OutJson = PropertyToJsonData(Property, &TestIn.TestObjectSoft);
-		TestObjectRefData(*this, OutJson->AsObject(), TestOther->GetPathName());
-
-		FToolsetJsonTest TestOut;
-		TestTrue("Output", JsonDataToProperty(OutJson, Property, &TestOut.TestObjectSoft));
-		TestEqual("RefCheck", TestIn.TestObjectSoft, TestOut.TestObjectSoft);
-	});
-
-	It("Converts WeakPtr property to reference schema", [this]()
-	{
-		UStruct* Struct = FToolsetJsonTest::StaticStruct();
-		FProperty* Prop = Struct->FindPropertyByName("TestObjectWeak");
-		const TSharedPtr<FJsonObject> Schema = PropertyToJsonSchema(Prop);
-		TestObjectRefSchema(*this, Schema, UToolsetJsonTestObject::StaticClass(), false);
-	});
-
-	It("Converts WeakPtr property to and from JSON", [this]()
-	{
-		FToolsetJsonTest TestIn;
-		TObjectPtr<UToolsetJsonTestObject> TestOther = NewObject<UToolsetJsonTestObject>();
-		TestIn.TestObjectWeak = TestOther;
-
-		UStruct* Struct = FToolsetJsonTest::StaticStruct();
-		FProperty* Property = Struct->FindPropertyByName("TestObjectWeak");
-
-		TSharedPtr<FJsonValue> OutJson = PropertyToJsonData(Property, &TestIn.TestObjectWeak);
-		TestObjectRefData(*this, OutJson->AsObject(), TestOther->GetPathName());
-
-		FToolsetJsonTest TestOut;
-		TestTrue("Output", JsonDataToProperty(OutJson, Property, &TestOut.TestObjectWeak));
-		TestEqual("RefCheck", TestIn.TestObjectWeak, TestOut.TestObjectWeak);
-	});
-
-	It("Converts UClass* property to reference schema", [this]()
-	{
-		UStruct* Struct = FToolsetJsonTest::StaticStruct();
-		FProperty* Prop = Struct->FindPropertyByName("TestClass");
-		const TSharedPtr<FJsonObject> Schema = PropertyToJsonSchema(Prop);
-		TestObjectRefSchema(*this, Schema, UToolsetJsonTestObject::StaticClass(), true);
-	});
-
-	It("Converts UClass* property to and from JSON", [this]()
-	{
-		FToolsetJsonTest TestIn;
-		TestIn.TestClass = UToolsetJsonTestObject::StaticClass();
-
-		UStruct* Struct = FToolsetJsonTest::StaticStruct();
-		FProperty* Property = Struct->FindPropertyByName("TestClass");
-
-		TSharedPtr<FJsonValue> OutJson = PropertyToJsonData(Property, &TestIn.TestClass);
-		TestObjectRefData(*this, OutJson->AsObject(), UToolsetJsonTestObject::StaticClass()->GetPathName());
-
-		FToolsetJsonTest TestOut;
-		TestTrue("Output", JsonDataToProperty(OutJson, Property, &TestOut.TestClass));
-		TestEqual("RefCheck", TestIn.TestClass, TestOut.TestClass);
-	});
-
-	It("Converts SoftClassPath property to reference schema", [this]()
-	{
-		UStruct* Struct = FToolsetJsonTest::StaticStruct();
-		FProperty* Prop = Struct->FindPropertyByName("TestClassPath");
-		const TSharedPtr<FJsonObject> Schema = PropertyToJsonSchema(Prop);
-		TestObjectRefSchema(*this, Schema, UObject::StaticClass(), true);
-	});
-
-	It("Converts SoftClassPath property to and from JSON", [this]()
-	{
-		FToolsetJsonTest TestIn;
-		TestIn.TestClassPath = UToolsetJsonTestObject::StaticClass();
-
-		UStruct* Struct = FToolsetJsonTest::StaticStruct();
-		FProperty* Property = Struct->FindPropertyByName("TestClassPath");
-
-		TSharedPtr<FJsonValue> OutJson = PropertyToJsonData(Property, &TestIn.TestClassPath);
-		TestObjectRefData(*this, OutJson->AsObject(), UToolsetJsonTestObject::StaticClass()->GetPathName());
-
-		FToolsetJsonTest TestOut;
-		TestTrue("Output", JsonDataToProperty(OutJson, Property, &TestOut.TestClassPath));
-		TestEqual("RefCheck", TestIn.TestClassPath, TestOut.TestClassPath);
-	});
-
-	It("Converts FColor property to FLinearColor schema", [this]()
-	{
-		UStruct* Struct = FToolsetJsonTest::StaticStruct();
-		FProperty* Prop = Struct->FindPropertyByName("TestColor");
-		const TSharedPtr<FJsonObject> Schema = PropertyToJsonSchema(Prop);
-		TestColorSchema(*this, Schema, true);
-	});
-
-	It("Converts FColor property to and from JSON", [this]()
-	{
-		FToolsetJsonTest TestIn;
-		TestIn.TestColor = FColor(10, 20, 30, 40);
-
-		UStruct* Struct = FToolsetJsonTest::StaticStruct();
-		FProperty* Property = Struct->FindPropertyByName("TestColor");
-
-		TSharedPtr<FJsonValue> OutJson = PropertyToJsonData(Property, &TestIn.TestColor);
-		TestColorData(*this, OutJson->AsObject(), TestIn.TestColor);
-
-		FToolsetJsonTest TestOut;
-		TestTrue("Output", JsonDataToProperty(OutJson, Property, &TestOut.TestColor));
-		TestEqual("ColorCheck", TestIn.TestColor, TestOut.TestColor);
-	});
-
-	It("Converts FLinearColor property to FLinearColor schema", [this]()
-	{
-		UStruct* Struct = FToolsetJsonTest::StaticStruct();
-		FProperty* Prop = Struct->FindPropertyByName("TestLinearColor");
-		const TSharedPtr<FJsonObject> Schema = PropertyToJsonSchema(Prop);
-		TestColorSchema(*this, Schema, false);
-	});
-
-	It("Converts FLinearColor property to and from JSON", [this]()
-	{
-		FToolsetJsonTest TestIn;
-		TestIn.TestLinearColor = FLinearColor(0.1, 0.2, 0.3, 0.4);
-
-		UStruct* Struct = FToolsetJsonTest::StaticStruct();
-		FProperty* Property = Struct->FindPropertyByName("TestLinearColor");
-
-		TSharedPtr<FJsonValue> OutJson = PropertyToJsonData(Property, &TestIn.TestLinearColor);
-		TestColorData(*this, OutJson->AsObject(), TestIn.TestLinearColor.ToFColor(false));
-
-		FToolsetJsonTest TestOut;
-		TestTrue("Output", JsonDataToProperty(OutJson, Property, &TestOut.TestLinearColor));
-		TestEqual("ColorCheck", TestIn.TestLinearColor, TestOut.TestLinearColor);
-	});
-
-	It("Converts FTransform property to transform schema", [this]()
-	{
-		UStruct* Struct = FToolsetJsonTest::StaticStruct();
-		FProperty* Prop = Struct->FindPropertyByName("TestTransform");
-		const TSharedPtr<FJsonObject> Schema = PropertyToJsonSchema(Prop);
-		TestTransformSchema(*this, Schema);
-	});
-
-	It("Converts FTransform property to and from JSON", [this]()
-	{
-		FToolsetJsonTest TestIn;
-		TestIn.TestTransform = FTransform(
-			FRotator(10.f, 20.f, 30.f).Quaternion(),
-			FVector(100.f, 200.f, 300.f),
-			FVector(2.f, 3.f, 4.f));
-
-		UStruct* Struct = FToolsetJsonTest::StaticStruct();
-		FProperty* Property = Struct->FindPropertyByName("TestTransform");
-
-		TSharedPtr<FJsonValue> OutJson = PropertyToJsonData(Property, &TestIn.TestTransform);
-		if (!TestTrue("Json", OutJson.IsValid())) return;
-
-		TSharedPtr<FJsonObject> JsonObject = OutJson->AsObject();
-		if (!TestTrue("Json object", JsonObject.IsValid())) return;
-
-		TestTrue("location", JsonObject->HasField(TEXT("location")));
-		TestTrue("rotation", JsonObject->HasField(TEXT("rotation")));
-		TestTrue("scale", JsonObject->HasField(TEXT("scale")));
-
-		FToolsetJsonTest TestOut;
-		TestTrue("Output", JsonDataToProperty(OutJson, Property, &TestOut.TestTransform));
-		TestEqual("Location", TestOut.TestTransform.GetTranslation(),
-			TestIn.TestTransform.GetTranslation());
-		TestEqual("Rotation", TestOut.TestTransform.GetRotation().Rotator(),
-			TestIn.TestTransform.GetRotation().Rotator());
-		TestEqual("Scale", TestOut.TestTransform.GetScale3D(),
-			TestIn.TestTransform.GetScale3D());
-	});
-
-	It("Partial FTransform JSON preserves unset fields", [this]()
-	{
-		// Verify JsonDataToProperty with a partial JSON object (only location) leaves
-		// the existing rotation and scale unchanged — matching to_ue semantics.
-		const FRotator OriginalRotation(45.f, 90.f, 0.f);
-		const FVector OriginalScale(2.f, 2.f, 2.f);
-
-		FToolsetJsonTest Test;
-		Test.TestTransform = FTransform(
-			OriginalRotation.Quaternion(), FVector::ZeroVector, OriginalScale);
-
-		UStruct* Struct = FToolsetJsonTest::StaticStruct();
-		FProperty* Property = Struct->FindPropertyByName("TestTransform");
-
-		const FVector NewLocation(1.f, 2.f, 3.f);
-		TSharedPtr<FJsonObject> LocationJson = MakeShared<FJsonObject>();
-		LocationJson->SetNumberField(TEXT("x"), NewLocation.X);
-		LocationJson->SetNumberField(TEXT("y"), NewLocation.Y);
-		LocationJson->SetNumberField(TEXT("z"), NewLocation.Z);
-
-		TSharedPtr<FJsonObject> PartialJson = MakeShared<FJsonObject>();
-		PartialJson->SetObjectField(TEXT("location"), LocationJson);
-
-		TestTrue("Output", JsonDataToProperty(
-			MakeShared<FJsonValueObject>(PartialJson), Property, &Test.TestTransform));
-		TestEqual("Location updated", Test.TestTransform.GetTranslation(), NewLocation);
-		TestEqual("Rotation preserved", Test.TestTransform.GetRotation().Rotator(), OriginalRotation);
-		TestEqual("Scale preserved", Test.TestTransform.GetScale3D(), OriginalScale);
-	});
-
-	It("Converts struct to reference schema", [this]()
-	{		
-		TSharedPtr<FJsonObject> Schema = StructToJsonSchema(FToolsetJsonTest::StaticStruct());
-		TestTrue("Convert", Schema.IsValid());
-		if (!Schema.IsValid())
-		{
-			return;
-		}
-		const TSharedPtr<FJsonObject>* Properties = nullptr;
-		TestTrue("Properties", Schema->TryGetObjectField(TEXT("properties"), Properties));
-		if (Properties)
-		{
-			for (TFieldIterator<FProperty> It(FToolsetJsonTest::StaticStruct()); It; ++It)
-			{
-				FProperty* Property = *It;
-				FString JsonPropertyName = FJsonObjectConverter::StandardizeCase(Property->GetName());
-				const TSharedPtr<FJsonObject>* TestProperty = nullptr;
-				TestTrue("Object", Properties->Get()->TryGetObjectField(JsonPropertyName, TestProperty));
-				if (TestProperty)
-				{
-					if (JsonPropertyName == FString(TEXT("testColor")) ||
-						JsonPropertyName == FString(TEXT("testLinearColor")))
-					{
-						TestColorSchema(*this, *TestProperty, JsonPropertyName == FString(TEXT("testColor")));
-					}
-					else if (JsonPropertyName == TEXT("testTransform"))
-					{
-						TestTransformSchema(*this, *TestProperty);
-					}
-					else if (JsonPropertyName == TEXT("testToolCallAsyncResult"))
-					{
-						continue;  // Ignored as references to the base class are always null.
-					}
-					else
-					{
-						TestObjectRefSchema(*this, *TestProperty, nullptr, false);
-					}
-				}
-			}
-		}
-	});
-
-	It("Converts struct to and from JSON.", [this]()
-	{
-		FToolsetJsonTest TestIn;
-		TObjectPtr<UToolsetJsonTestObject> TestOther = NewObject<UToolsetJsonTestObject>();
-		TestIn.TestObject = TestOther;
-		TestIn.TestObjectPath = TestOther;
-		TestIn.TestObjectSoft = TestOther;
-		TestIn.TestObjectWeak = TestOther;
-		TestIn.TestClass = UToolsetJsonTestObject::StaticClass();
-		TestIn.TestClassPath = UToolsetJsonTestObject::StaticClass();
-		TestIn.TestColor = FColor(10, 20, 30, 40);
-		TestIn.TestTransform = FTransform(
-			FRotator(10.f, 20.f, 30.f).Quaternion(),
-			FVector(100.f, 200.f, 300.f),
-			FVector(2.f, 3.f, 4.f));
-
-		TSharedPtr<FJsonObject> JsonObject = StructToJsonData(
-			FToolsetJsonTest::StaticStruct(), &TestIn);
-		for (TFieldIterator<FProperty> It(FToolsetJsonTest::StaticStruct()); It; ++It)
-		{
-			FProperty* Property = *It;
-			FString JsonPropertyName = FJsonObjectConverter::StandardizeCase(Property->GetName());
-			if (JsonPropertyName == "testClass" || JsonPropertyName == "testClassPath")
-			{
-				TestObjectRefDataInStruct(
-					*this, JsonObject, JsonPropertyName, UToolsetJsonTestObject::StaticClass()->GetPathName());
-			}
-			else if (JsonPropertyName == "testColor")
-			{
-				TestColorDataInStruct(
-					*this, JsonObject, JsonPropertyName, TestIn.TestColor);
-			}
-			else if (JsonPropertyName == "testLinearColor")
-			{
-				TestColorDataInStruct(
-					*this, JsonObject, JsonPropertyName, TestIn.TestLinearColor.ToFColor(false));
-			}
-			else if (JsonPropertyName == TEXT("testTransform"))
-			{
-				TestTrue("testTransform", JsonObject->HasField(TEXT("testTransform")));
-			}
-			else if (JsonPropertyName == TEXT("testToolCallAsyncResult"))
-			{
-				continue;  // Does not support conversion.
-			}
-			else
-			{
-				TestObjectRefDataInStruct(
-					*this, JsonObject, JsonPropertyName, TestOther->GetPathName());
-			}
-		}
-
-		FToolsetJsonTest TestOut;
-		TestTrue("FromJson", JsonDataToStruct(
-			JsonObject.ToSharedRef(), FToolsetJsonTest::StaticStruct(), &TestOut));
-		TestEqual("RefCheck", TestIn.TestObject, TestOut.TestObject);
-		TestEqual("RefCheck", TestIn.TestObjectPath, TestOut.TestObjectPath);
-		TestEqual("RefCheck", TestIn.TestObjectSoft, TestOut.TestObjectSoft);
-		TestEqual("RefCheck", TestIn.TestObjectWeak, TestOut.TestObjectWeak);
-		TestEqual("RefCheck", TestIn.TestClass, TestOut.TestClass);
-		TestEqual("RefCheck", TestIn.TestClassPath, TestOut.TestClassPath);
-		TestEqual("ColorCheck", TestIn.TestColor, TestOut.TestColor);
-		TestEqual("TransformLocation", TestIn.TestTransform.GetTranslation(),
-			TestOut.TestTransform.GetTranslation());
-		TestEqual("TransformRotation", TestIn.TestTransform.GetRotation().Rotator(),
-			TestOut.TestTransform.GetRotation().Rotator());
-		TestEqual("TransformScale", TestIn.TestTransform.GetScale3D(),
-			TestOut.TestTransform.GetScale3D());
-	});
-
-	It("Imports from old string style path", [this]()
-	{
-		TObjectPtr<UToolsetJsonTestObject> TestOther = NewObject<UToolsetJsonTestObject>();
-		UStruct* Struct = FToolsetJsonTest::StaticStruct();
-		FProperty* Property = Struct->FindPropertyByName("TestObject");
-
-		FToolsetJsonTest TestOut;
-		TSharedPtr<FJsonValue> OutJson = MakeShared<TJsonValueString<TCHAR>>(TestOther->GetPathName());
-		TestTrue("Output", JsonDataToProperty(OutJson, Property, &TestOut.TestObject));
-		TestEqual("RefCheck", TestOut.TestObject, TestOther);
-	});
-
 	It("Filters user visible properties schema", [this]()
 	{
 		TSharedPtr<FJsonObject> ClassSchema = StructToJsonSchema(
 			UToolsetJsonTestObject::StaticClass(), true);
-		TestTrue("TestData", ClassSchema.IsValid());
-		if (ClassSchema)
+		if (!TestTrue("TestData", ClassSchema.IsValid())) return;
+		const TSharedPtr<FJsonObject>* Properties = nullptr;
+		if (!(TestTrue("Properties", ClassSchema->TryGetObjectField(TEXT("properties"), Properties)) &&
+			Properties))
 		{
-			const TSharedPtr<FJsonObject>* Properties = nullptr;
-			TestTrue("Properties", ClassSchema->TryGetObjectField(TEXT("properties"), Properties));
-			if (Properties)
-			{
-				TestTrue("Visible", Properties->Get()->HasField(TEXT("visibleProperty")));
-				TestFalse("Hidden", Properties->Get()->HasField(TEXT("hiddenProperty")));
-			}
+			return;
 		}
+		TestTrue("Visible", Properties->Get()->HasField(TEXT("visibleProperty")));
+		TestFalse("Hidden", Properties->Get()->HasField(TEXT("hiddenProperty")));
 	});
 
 	It("Filters user visible properties data", [this]()
 	{
 		TSharedPtr<FJsonObject> TestData = StructToJsonData(
 			UToolsetJsonTestObject::StaticClass(), NewObject<UToolsetJsonTestObject>(), true);
-		TestTrue("TestData", TestData.IsValid());
-		if (TestData)
-		{
-			TestTrue("Visible", TestData->HasField(TEXT("visibleProperty")));
-			TestFalse("Hidden", TestData->HasField(TEXT("hiddenProperty")));
-		}
-	});
-
-	It("Can read defaults", [this]()
-	{
-		UFunction* FunctionPtr = UToolsetJsonTestObject::StaticClass()->FindFunctionByName("TestDefault");
-		const TSharedPtr<FJsonObject> Schema = StructToJsonSchema(FunctionPtr);
-		TestTrue("Schema", Schema.IsValid());
-		if (!Schema.IsValid())
-		{
-			return;
-		}
-		const TSharedPtr<FJsonObject>* InputSchema = nullptr;
-		TestTrue("Input", Schema->TryGetObjectField(TEXT("inputSchema"), InputSchema));
-		if (InputSchema)
-		{
-			TestFalse("NoRequired", InputSchema->Get()->HasField(TEXT("required")));
-			const TSharedPtr<FJsonObject>* Properties = nullptr;
-			TestTrue("Properties", InputSchema->Get()->TryGetObjectField(TEXT("properties"), Properties));
-			if (Properties)
-			{
-				const TSharedPtr<FJsonObject>* TestObject = nullptr;
-				TestTrue("Input", Properties->Get()->TryGetObjectField(TEXT("testObject"), TestObject));
-				if (TestObject)
-				{
-					TestTrue("Default", TestObject->Get()->GetField(TEXT("default"), EJson::Null)->IsNull());
-				}
-
-				const TSharedPtr<FJsonObject>* TestColor = nullptr;
-				TestTrue("Input", Properties->Get()->TryGetObjectField(TEXT("testColor"), TestColor));
-				if (TestColor)
-				{
-					TestColorData(
-						*this, TestColor->Get()->GetObjectField(TEXT("default")), FColor(10, 20, 30, 40));
-				}
-
-				const TSharedPtr<FJsonObject>* TestLinearColor = nullptr;
-				TestTrue("Input", Properties->Get()->TryGetObjectField(TEXT("testLinearColor"), TestLinearColor));
-				if (TestLinearColor)
-				{
-					TestColorData(
-						*this, TestLinearColor->Get()->GetObjectField(TEXT("default")),
-						FLinearColor(0.1, 0.2, 0.3, 0.4).ToFColor(false));
-				}
-			}
-		}
+		if (!TestTrue("TestData", TestData.IsValid())) return;
+		TestTrue("Visible", TestData->HasField(TEXT("visibleProperty")));
+		TestFalse("Hidden", TestData->HasField(TEXT("hiddenProperty")));
 	});
 
 	It(TEXT("Converts ToolCallAsyncResult return value to JSON schema"), [this]()
@@ -982,7 +94,6 @@ void FToolsetJsonSpec::Define()
 			TEXT("Matches expected schema"),
 			JsonToString(ReturnValueSchema.ToSharedRef()),
 			JsonToString(PropertyToJsonSchema(Function->GetReturnProperty()).ToSharedRef()));
-			
 	});
 
 	It(TEXT("Converts ToolCallAsyncResult to JSON"), [this]()
diff --git a/Engine/Plugins/Experimental/ToolsetRegistry/Source/ToolsetRegistry/Private/Tests/ToolsetJsonTest.h b/Engine/Plugins/Experimental/ToolsetRegistry/Source/ToolsetRegistry/Private/Tests/ToolsetJsonTest.h
--- a/Engine/Plugins/Experimental/ToolsetRegistry/Source/ToolsetRegistry/Private/Tests/ToolsetJsonTest.h
+++ b/Engine/Plugins/Experimental/ToolsetRegistry/Source/ToolsetRegistry/Private/Tests/ToolsetJsonTest.h
@@ -3,9 +3,6 @@
 #pragma once
 
 #include "CoreMinimal.h"
-#include "Kismet/KismetSystemLibrary.h"
-#include "Templates/SharedPointer.h"
-#include "Templates/SubclassOf.h"
 
 #include "ToolsetRegistry/ToolCallAsyncResult.h"
 
@@ -35,56 +32,6 @@ public:
 	// CPF_Edit | CPF_DisableEditOnTemplate → accessible for get, CannotEditTemplate on templates (CDOs).
 	UPROPERTY(EditInstanceOnly, Category="ToolsetTest")
 	bool InstanceOnlyProperty;
-
-	UFUNCTION(BlueprintCallable, Category = "Main")
-	void TestDefault(
-		UObject* TestObject = nullptr, FColor TestColor = FColor(10, 20, 30, 40),
-		FLinearColor TestLinearColor = FLinearColor(0.1, 0.2, 0.3, 0.4)) {}
-
-	UFUNCTION(BlueprintCallable, Category = "ToolsetTest")
-	void TestOptionalObject(
-		UObject* TestObject = nullptr) {}
-
-	UFUNCTION(BlueprintCallable, Category = "ToolsetTest")
-	void TestObjectParam(
-		UObject* TestObject) {}
-
-	UFUNCTION(BlueprintCallable, Category = "ToolsetTest")
-	void TestSpecObjectParam(
-		UToolsetJsonTestObject* TestObject) {}
-
-	UFUNCTION(BlueprintCallable, Category = "ToolsetTest")
-	void TestClassParam(
-		UClass* TestClass) {}
-
-	UFUNCTION(BlueprintCallable, Category = "ToolsetTest")
-	void TestSubClassParam(
-		TSubclassOf<UToolsetJsonTestObject> TestClass) {}
-
-	UFUNCTION(BlueprintCallable, Category = "ToolsetTest")
-	void TestSoftObjectParam(
-		TSoftObjectPtr<UToolsetJsonTestObject> TestObject) {}
-
-	UFUNCTION(BlueprintCallable, Category = "ToolsetTest")
-	void TestSoftClassParam(
-		TSoftClassPtr<UToolsetJsonTestObject> TestClass) {}
-
-	UFUNCTION(BlueprintCallable, Category = "ToolsetTest")
-	void TestSoftClassPathParam(
-		FSoftClassPath TestClassPath) {}
-
-	// Functions for testing FToolCallScriptExceptionHandler
-	UFUNCTION(BlueprintCallable, Category = "ToolsetTest")
-	void TestMultipleRefs(
-		UToolsetJsonTestObject* FirstObject,
-		UToolsetJsonTestObject* SecondObject) {}
-
-	UFUNCTION(BlueprintCallable, Category = "ToolsetTest")
-	void TestEmptyError()
-	{
-		// Raise an explicitly empty error
-		UKismetSystemLibrary::RaiseScriptError(FString());
-	}
 };
 
 USTRUCT(BlueprintType)
@@ -92,33 +39,6 @@ struct FToolsetJsonTest
 {
 	GENERATED_BODY()
 public:
-	UPROPERTY()
-	TObjectPtr<UToolsetJsonTestObject> TestObject;
-
-	UPROPERTY()
-	FSoftObjectPath TestObjectPath;
-
-	UPROPERTY()
-	TSoftObjectPtr<UToolsetJsonTestObject> TestObjectSoft;
-
-	UPROPERTY()
-	TWeakObjectPtr<UToolsetJsonTestObject> TestObjectWeak;
-
-	UPROPERTY()
-	TSubclassOf<UToolsetJsonTestObject> TestClass;
-
-	UPROPERTY()
-	FSoftClassPath TestClassPath;
-
-	UPROPERTY()
-	FColor TestColor = FColor::White;
-
-	UPROPERTY()
-	FLinearColor TestLinearColor = FLinearColor::White;
-
-	UPROPERTY()
-	FTransform TestTransform;
-
 	UPROPERTY()
 	TObjectPtr<UToolCallAsyncResult> TestToolCallAsyncResult;
 };
diff --git a/Engine/Plugins/Experimental/ToolsetRegistry/Source/ToolsetRegistry/Private/Tests/TransformConverterTest.cpp b/Engine/Plugins/Experimental/ToolsetRegistry/Source/ToolsetRegistry/Private/Tests/TransformConverterTest.cpp
new file mode 100644
--- /dev/null
+++ b/Engine/Plugins/Experimental/ToolsetRegistry/Source/ToolsetRegistry/Private/Tests/TransformConverterTest.cpp
@@ -0,0 +1,116 @@
+// Copyright Epic Games, Inc. All Rights Reserved.
+
+#include "TransformConverterTest.h"
+
+#include "CoreMinimal.h"
+#include "Dom/JsonObject.h"
+#include "Misc/AutomationTest.h"
+
+#include "ToolsetRegistry/ToolsetJson.h"
+
+#if WITH_DEV_AUTOMATION_TESTS
+
+BEGIN_DEFINE_SPEC(FToolsetTransformConverterSpec, "AI.ToolsetRegistry.TransformConverterSpec",
+	EAutomationTestFlags::EditorContext | EAutomationTestFlags::EngineFilter)
+	void TestTransformSchema(const TSharedPtr<FJsonObject>& Schema);
+END_DEFINE_SPEC(FToolsetTransformConverterSpec)
+
+void FToolsetTransformConverterSpec::TestTransformSchema(const TSharedPtr<FJsonObject>& Schema)
+{
+	if (!TestTrue("Schema", Schema.IsValid())) return;
+
+	FString Type;
+	TestTrue("Type", Schema->TryGetStringField(TEXT("type"), Type));
+	TestEqual("Type", Type, FString(TEXT("object")));
+
+	TestFalse("No required", Schema->HasField(TEXT("required")));
+
+	const TSharedPtr<FJsonObject>* Properties = nullptr;
+	if (!(TestTrue("Properties", Schema->TryGetObjectField(TEXT("properties"), Properties)) &&
+		Properties))
+	{
+		return;
+	}
+	for (const FString& Field : TArray<FString>{ TEXT("location"), TEXT("rotation"), TEXT("scale") })
+	{
+		const TSharedPtr<FJsonObject>* SubSchema = nullptr;
+		TestTrue(*Field, Properties->Get()->TryGetObjectField(Field, SubSchema));
+	}
+}
+
+void FToolsetTransformConverterSpec::Define()
+{
+	using namespace UE::ToolsetRegistry::Internal;
+	using namespace UE::ToolsetRegistry::Internal::ToolsetJson;
+
+	It("Converts FTransform property to transform schema", [this]()
+	{
+		UStruct* Struct = FTransformConverterTest::StaticStruct();
+		FProperty* Prop = Struct->FindPropertyByName("TestTransform");
+		const TSharedPtr<FJsonObject> Schema = PropertyToJsonSchema(Prop);
+		TestTransformSchema(Schema);
+	});
+
+	It("Converts FTransform property to and from JSON", [this]()
+	{
+		FTransformConverterTest TestIn;
+		TestIn.TestTransform = FTransform(
+			FRotator(10.f, 20.f, 30.f).Quaternion(),
+			FVector(100.f, 200.f, 300.f),
+			FVector(2.f, 3.f, 4.f));
+
+		UStruct* Struct = FTransformConverterTest::StaticStruct();
+		FProperty* Property = Struct->FindPropertyByName("TestTransform");
+
+		TSharedPtr<FJsonValue> OutJson = PropertyToJsonData(Property, &TestIn.TestTransform);
+		if (!TestTrue("Json", OutJson.IsValid())) return;
+
+		TSharedPtr<FJsonObject> JsonObject = OutJson->AsObject();
+		if (!TestTrue("Json object", JsonObject.IsValid())) return;
+
+		TestTrue("location", JsonObject->HasField(TEXT("location")));
+		TestTrue("rotation", JsonObject->HasField(TEXT("rotation")));
+		TestTrue("scale", JsonObject->HasField(TEXT("scale")));
+
+		FTransformConverterTest TestOut;
+		TestTrue("Output", JsonDataToProperty(OutJson, Property, &TestOut.TestTransform));
+		TestEqual("Location", TestOut.TestTransform.GetTranslation(),
+			TestIn.TestTransform.GetTranslation());
+		TestEqual("Rotation", TestOut.TestTransform.GetRotation().Rotator(),
+			TestIn.TestTransform.GetRotation().Rotator());
+		TestEqual("Scale", TestOut.TestTransform.GetScale3D(),
+			TestIn.TestTransform.GetScale3D());
+	});
+
+	It("Partial FTransform JSON preserves unset fields", [this]()
+	{
+		// Verify JsonDataToProperty with a partial JSON object (only location) leaves
+		// the existing rotation and scale unchanged — matching to_ue semantics.
+		const FRotator OriginalRotation(45.f, 90.f, 0.f);
+		const FVector OriginalScale(2.f, 2.f, 2.f);
+
+		FTransformConverterTest Test;
+		Test.TestTransform = FTransform(
+			OriginalRotation.Quaternion(), FVector::ZeroVector, OriginalScale);
+
+		UStruct* Struct = FTransformConverterTest::StaticStruct();
+		FProperty* Property = Struct->FindPropertyByName("TestTransform");
+
+		const FVector NewLocation(1.f, 2.f, 3.f);
+		TSharedPtr<FJsonObject> LocationJson = MakeShared<FJsonObject>();
+		LocationJson->SetNumberField(TEXT("x"), NewLocation.X);
+		LocationJson->SetNumberField(TEXT("y"), NewLocation.Y);
+		LocationJson->SetNumberField(TEXT("z"), NewLocation.Z);
+
+		TSharedPtr<FJsonObject> PartialJson = MakeShared<FJsonObject>();
+		PartialJson->SetObjectField(TEXT("location"), LocationJson);
+
+		TestTrue("Output", JsonDataToProperty(
+			MakeShared<FJsonValueObject>(PartialJson), Property, &Test.TestTransform));
+		TestEqual("Location updated", Test.TestTransform.GetTranslation(), NewLocation);
+		TestEqual("Rotation preserved", Test.TestTransform.GetRotation().Rotator(), OriginalRotation);
+		TestEqual("Scale preserved", Test.TestTransform.GetScale3D(), OriginalScale);
+	});
+}
+
+#endif
diff --git a/Engine/Plugins/Experimental/ToolsetRegistry/Source/ToolsetRegistry/Private/Tests/TransformConverterTest.h b/Engine/Plugins/Experimental/ToolsetRegistry/Source/ToolsetRegistry/Private/Tests/TransformConverterTest.h
new file mode 100644
--- /dev/null
+++ b/Engine/Plugins/Experimental/ToolsetRegistry/Source/ToolsetRegistry/Private/Tests/TransformConverterTest.h
@@ -0,0 +1,17 @@
+// Copyright Epic Games, Inc. All Rights Reserved.
+
+#pragma once
+
+#include "CoreMinimal.h"
+
+#include "TransformConverterTest.generated.h"
+
+
+USTRUCT(BlueprintType)
+struct FTransformConverterTest
+{
+	GENERATED_BODY()
+public:
+	UPROPERTY()
+	FTransform TestTransform;
+};
