Fix a few some reference converter test params.

#Synced-CL 52087183

diff --git a/Engine/Plugins/Experimental/ToolsetRegistry/Source/ToolsetRegistry/Private/Tests/ReferenceConverterTest.cpp b/Engine/Plugins/Experimental/ToolsetRegistry/Source/ToolsetRegistry/Private/Tests/ReferenceConverterTest.cpp
--- a/Engine/Plugins/Experimental/ToolsetRegistry/Source/ToolsetRegistry/Private/Tests/ReferenceConverterTest.cpp
+++ b/Engine/Plugins/Experimental/ToolsetRegistry/Source/ToolsetRegistry/Private/Tests/ReferenceConverterTest.cpp
@@ -490,8 +490,8 @@ void FToolsetReferenceConverterSpec::Define()
 		if (!TestTrue("Schema", Schema.IsValid())) return;
 
 		const TSharedPtr<FJsonObject>* InputSchema = nullptr;
-		if (!(TestTrue("Input", Schema->TryGetObjectField(TEXT("inputSchema"), InputSchema) &&
-			InputSchema)))
+		if (!(TestTrue("Input", Schema->TryGetObjectField(TEXT("inputSchema"), InputSchema)) &&
+			InputSchema))
 		{
 			return;
 		}
@@ -499,15 +499,15 @@ void FToolsetReferenceConverterSpec::Define()
 		TestFalse("NoRequired", InputSchema->Get()->HasField(TEXT("required")));
 
 		const TSharedPtr<FJsonObject>* Properties = nullptr;
-		if (!(TestTrue("Properties", InputSchema->Get()->TryGetObjectField(TEXT("properties"), Properties) &&
-			Properties)))
+		if (!(TestTrue("Properties", InputSchema->Get()->TryGetObjectField(TEXT("properties"), Properties)) &&
+			Properties))
 		{
 			return;
 		}
 		
 		const TSharedPtr<FJsonObject>* TestObject = nullptr;
-		if (!(TestTrue("testObject", Properties->Get()->TryGetObjectField(TEXT("testObject"), TestObject) &&
-			TestObject)))
+		if (!(TestTrue("testObject", Properties->Get()->TryGetObjectField(TEXT("testObject"), TestObject)) &&
+			TestObject))
 		{
 			return;
 		}
