[AI Assistant] Make viewport screenshot tool more robust

When tests execute or the editor runs with the viewport in the
"background" the editor viewport isn't rendered and so scheduled
screenshots are never executed. This works around the issue by
marking the active editor viewport as "realtime" which forces
it to be rendered for the duration of the viewport screenshot tool
execution.

#Synced-CL 52101884

diff --git a/Engine/Plugins/Experimental/ToolsetRegistry/Source/ToolsetRegistry/Private/ToolsetRegistry/EditorAppToolset.cpp b/Engine/Plugins/Experimental/ToolsetRegistry/Source/ToolsetRegistry/Private/ToolsetRegistry/EditorAppToolset.cpp
--- a/Engine/Plugins/Experimental/ToolsetRegistry/Source/ToolsetRegistry/Private/ToolsetRegistry/EditorAppToolset.cpp
+++ b/Engine/Plugins/Experimental/ToolsetRegistry/Source/ToolsetRegistry/Private/ToolsetRegistry/EditorAppToolset.cpp
@@ -14,13 +14,17 @@
 #include "Engine/Texture.h"
 #include "HAL/ConsoleManager.h"
 #include "Kismet/KismetSystemLibrary.h"
+#include "LevelEditor.h"
+#include "LevelEditorViewport.h"
 #include "Materials/Material.h"
 #include "Materials/MaterialInterface.h"
+#include "Modules/ModuleManager.h"
 #include "Misc/ObjectThumbnail.h"
 #include "Misc/PackageName.h"
 #include "ObjectTools.h"
 #include "RenderingThread.h"
 #include "ShaderCompiler.h"
+#include "SLevelViewport.h"
 #include "Templates/SharedPointer.h"
 #include "ToolsetRegistry/DelegateHandle.h"
 #include "ToolsetRegistry/JsonConversion.h"
@@ -38,11 +42,18 @@ namespace
 	// Helper that captures a single viewport screenshot and resolves a result.
 	// The object is kept alive by a strong reference captured in the screenshot
 	// delegate lambda, which is released once the callback fires.
-	struct FViewportScreenshotCapture : public TSharedFromThis<FViewportScreenshotCapture>
+	class FViewportScreenshotCapture : public TSharedFromThis<FViewportScreenshotCapture>
 	{
 	public:
 		FViewportScreenshotCapture() : Result(NewObject<UToolCallAsyncResultImage>()) {}
-		~FViewportScreenshotCapture() = default;
+		~FViewportScreenshotCapture()
+		{
+			if (LevelEditorViewport.IsValid())
+			{
+				LevelEditorViewport->GetLevelViewportClient().RemoveRealtimeOverride(
+					FText::FromString(Label));
+			}
+		}
 
 		static TStrongObjectPtr<UToolCallAsyncResultImage> Start()
 		{
@@ -52,6 +63,11 @@ namespace
 	private:
 		TStrongObjectPtr<UToolCallAsyncResultImage> StartCapture()
 		{
+			if (!ConfigureViewport())
+			{
+				Result->SetError(TEXT("Editor viewport is not available."));
+				return Result;
+			}
 			FOnScreenshotCaptured& OnScreenshotCaptured =
 				FScreenshotRequest::OnScreenshotCaptured();
 			Handle =
@@ -68,6 +84,24 @@ namespace
 			return Result;
 		}
 
+		bool ConfigureViewport()
+		{
+			check(!LevelEditorViewport.IsValid());
+			FLevelEditorModule* LevelEditor =
+					FModuleManager::Get().IsModuleLoaded(TEXT("LevelEditor"))
+						? &FModuleManager::GetModuleChecked<FLevelEditorModule>(
+							TEXT("LevelEditor"))
+						: nullptr;
+			if (!LevelEditor) return false;
+
+			LevelEditorViewport = LevelEditor->GetFirstActiveLevelViewport();
+			if (!LevelEditorViewport) return false;
+
+			LevelEditorViewport->GetLevelViewportClient().AddRealtimeOverride(
+				true, FText::FromString(Label));
+			return true;
+		}
+
 		void OnCaptured(int32 Width, int32 Height, const TArray<FColor>& Colors)
 		{
 			const int32 MaxSize = 1280;
@@ -102,8 +136,13 @@ namespace
 			}
 		}
 
+	private:
 		TStrongObjectPtr<UToolCallAsyncResultImage> Result;
 		UE::ToolsetRegistry::FDelegateHandleRaii Handle;
+		TSharedPtr<SLevelViewport> LevelEditorViewport;
+
+	private:
+		static constexpr TCHAR Label[] = TEXT("ViewportScreenshotCapture");
 	};
 
 	// Helper that renders an asset thumbnail and resolves a result.
