set_properties in ObjectTools can now set instanced sub-object properties.

#Synced-CL 52154636

diff --git a/Engine/Plugins/Experimental/ToolsetRegistry/Source/ToolsetRegistry/Private/Tests/ReferenceConverterTest.cpp b/Engine/Plugins/Experimental/ToolsetRegistry/Source/ToolsetRegistry/Private/Tests/ReferenceConverterTest.cpp
--- a/Engine/Plugins/Experimental/ToolsetRegistry/Source/ToolsetRegistry/Private/Tests/ReferenceConverterTest.cpp
+++ b/Engine/Plugins/Experimental/ToolsetRegistry/Source/ToolsetRegistry/Private/Tests/ReferenceConverterTest.cpp
@@ -483,6 +483,63 @@ void FToolsetReferenceConverterSpec::Define()
 		TestEqual("RefCheck", TestOut.TestObject, TestOther);
 	});
 
+	It("Creates instanced sub-object from class refPath", [this]()
+	{
+		TObjectPtr<UReferenceConverterTestObject> Outer = NewObject<UReferenceConverterTestObject>();
+		FProperty* Property = UReferenceConverterTestObject::StaticClass()->FindPropertyByName(TEXT("TestInstancedObject"));
+		void* PropPtr = Property->ContainerPtrToValuePtr<void>(Outer.Get());
+
+		TSharedPtr<FJsonObject> JsonRef = MakeShared<FJsonObject>();
+		JsonRef->SetStringField(TEXT("refPath"), UReferenceConverterTestObject::StaticClass()->GetPathName());
+		TestTrue("Import", JsonDataToProperty(MakeShared<FJsonValueObject>(JsonRef), Property, PropPtr));
+
+		UObject* Created = CastField<FObjectProperty>(Property)->GetObjectPropertyValue(PropPtr);
+		if (!TestNotNull("Instance created", Created)) return;
+		TestTrue("Correct class", Created->IsA(UReferenceConverterTestObject::StaticClass()));
+		TestTrue("Outer is correct", Created->GetOuter() == Outer.Get());
+	});
+
+	It("Reuses existing instanced object when class matches", [this]()
+	{
+		TObjectPtr<UReferenceConverterTestObject> Outer = NewObject<UReferenceConverterTestObject>();
+		FProperty* Property = UReferenceConverterTestObject::StaticClass()->FindPropertyByName(TEXT("TestInstancedObject"));
+		void* PropPtr = Property->ContainerPtrToValuePtr<void>(Outer.Get());
+
+		TSharedPtr<FJsonObject> JsonRef = MakeShared<FJsonObject>();
+		JsonRef->SetStringField(TEXT("refPath"), UReferenceConverterTestObject::StaticClass()->GetPathName());
+
+		TestTrue("First import", JsonDataToProperty(MakeShared<FJsonValueObject>(JsonRef), Property, PropPtr));
+		UObject* First = CastField<FObjectProperty>(Property)->GetObjectPropertyValue(PropPtr);
+		TestTrue("Second import", JsonDataToProperty(MakeShared<FJsonValueObject>(JsonRef), Property, PropPtr));
+		UObject* Second = CastField<FObjectProperty>(Property)->GetObjectPropertyValue(PropPtr);
+		TestEqual("Same instance reused", First, Second);
+	});
+
+	It("Does not create instance for invalid class path", [this]()
+	{
+		TObjectPtr<UReferenceConverterTestObject> Outer = NewObject<UReferenceConverterTestObject>();
+		FProperty* Property = UReferenceConverterTestObject::StaticClass()->FindPropertyByName(TEXT("TestInstancedObject"));
+		void* PropPtr = Property->ContainerPtrToValuePtr<void>(Outer.Get());
+
+		TSharedPtr<FJsonObject> JsonRef = MakeShared<FJsonObject>();
+		JsonRef->SetStringField(TEXT("refPath"), TEXT("/Invalid/Path/DoesNotExist"));
+		TestTrue("Import", JsonDataToProperty(MakeShared<FJsonValueObject>(JsonRef), Property, PropPtr));
+		TestNull("No instance created", CastField<FObjectProperty>(Property)->GetObjectPropertyValue(PropPtr));
+	});
+
+	It("Does not create instance for wrong class type", [this]()
+	{
+		TObjectPtr<UReferenceConverterTestObject> Outer = NewObject<UReferenceConverterTestObject>();
+		FProperty* Property = UReferenceConverterTestObject::StaticClass()->FindPropertyByName(TEXT("TestInstancedObject"));
+		void* PropPtr = Property->ContainerPtrToValuePtr<void>(Outer.Get());
+
+		// UObject is not a subclass of UReferenceConverterTestObject
+		TSharedPtr<FJsonObject> JsonRef = MakeShared<FJsonObject>();
+		JsonRef->SetStringField(TEXT("refPath"), UObject::StaticClass()->GetPathName());
+		TestTrue("Import", JsonDataToProperty(MakeShared<FJsonValueObject>(JsonRef), Property, PropPtr));
+		TestNull("No instance created", CastField<FObjectProperty>(Property)->GetObjectPropertyValue(PropPtr));
+	});
+
 	It("Can read null default for object reference", [this]()
 	{
 		UFunction* FunctionPtr = UReferenceConverterTestObject::StaticClass()->FindFunctionByName("TestDefault");
diff --git a/Engine/Plugins/Experimental/ToolsetRegistry/Source/ToolsetRegistry/Private/Tests/ReferenceConverterTest.h b/Engine/Plugins/Experimental/ToolsetRegistry/Source/ToolsetRegistry/Private/Tests/ReferenceConverterTest.h
--- a/Engine/Plugins/Experimental/ToolsetRegistry/Source/ToolsetRegistry/Private/Tests/ReferenceConverterTest.h
+++ b/Engine/Plugins/Experimental/ToolsetRegistry/Source/ToolsetRegistry/Private/Tests/ReferenceConverterTest.h
@@ -13,6 +13,9 @@ class UReferenceConverterTestObject : public UObject
 {
 	GENERATED_BODY()
 public:
+	UPROPERTY(Instanced)
+	TObjectPtr<UReferenceConverterTestObject> TestInstancedObject;
+
 	UFUNCTION(BlueprintCallable, Category="ReferenceConverterTest")
 	void TestOptionalObject(UObject* TestObject = nullptr) {}
 
diff --git a/Engine/Plugins/Experimental/ToolsetRegistry/Source/ToolsetRegistry/Private/ToolsetRegistry/ReferenceConverter.cpp b/Engine/Plugins/Experimental/ToolsetRegistry/Source/ToolsetRegistry/Private/ToolsetRegistry/ReferenceConverter.cpp
--- a/Engine/Plugins/Experimental/ToolsetRegistry/Source/ToolsetRegistry/Private/ToolsetRegistry/ReferenceConverter.cpp
+++ b/Engine/Plugins/Experimental/ToolsetRegistry/Source/ToolsetRegistry/Private/ToolsetRegistry/ReferenceConverter.cpp
@@ -78,8 +78,7 @@ namespace UE::ToolsetRegistry
 
 	bool FToolsetReferenceConverter::CanConvertProperty(TNotNull<const FProperty*> Property)
 	{
-		if (CastField<FObjectProperty>(Property) &&
-			!Property->HasAnyPropertyFlags(CPF_PersistentInstance))
+		if (CastField<FObjectProperty>(Property))
 		{
 			return true;
 		}
@@ -292,19 +291,44 @@ namespace UE::ToolsetRegistry
 			}
 			else if (ObjectProperty)
 			{
-				UObject* Object = FSoftObjectPath(RefPath).ResolveObject();
-				if (!Object)
-				{
-					Object = FSoftObjectPath(RefPath).TryLoad();
-				}
-				if (!Object || !Object->IsA(ObjectProperty->PropertyClass))
+				if (ObjectProperty->HasAnyPropertyFlags(CPF_PersistentInstance))
 				{
-					UKismetSystemLibrary::RaiseScriptError(
-						FString::Printf(
-							TEXT("%s is not valid %s for property '%s'"),
+					// refPath is a class path — create a new owned instance
+					UClass* SubClass = FSoftClassPath(RefPath).TryLoadClass<UObject>();
+					if (!SubClass || !SubClass->IsChildOf(ObjectProperty->PropertyClass))
+					{
+						UKismetSystemLibrary::RaiseScriptError(FString::Printf(
+							TEXT("%s is not a valid subclass of %s for property '%s'"),
 							*RefPath, *ObjectProperty->PropertyClass->GetName(), *Property->GetName()));
+						return true;
+					}
+					// Back-calculate the outer: OutValue = (uint8*)Outer + Property->Offset
+					UObject* Outer = reinterpret_cast<UObject*>(
+						reinterpret_cast<uint8*>(OutValue) - Property->GetOffset_ForInternal());
+					UObject* Existing = ObjectProperty->GetObjectPropertyValue(OutValue);
+					if (!Existing || Existing->GetClass() != SubClass)
+					{
+						Existing = NewObject<UObject>(Outer, SubClass);
+						ObjectProperty->SetObjectPropertyValue(OutValue, Existing);
+					}
 					return true;
 				}
+				else
+				{
+					UObject* Object = FSoftObjectPath(RefPath).ResolveObject();
+					if (!Object)
+					{
+						Object = FSoftObjectPath(RefPath).TryLoad();
+					}
+					if (!Object || !Object->IsA(ObjectProperty->PropertyClass))
+					{
+						UKismetSystemLibrary::RaiseScriptError(
+							FString::Printf(
+								TEXT("%s is not valid %s for property '%s'"),
+								*RefPath, *ObjectProperty->PropertyClass->GetName(), *Property->GetName()));
+						return true;
+					}
+				}
 			}
 			// Validate soft references using Asset Registry
 			else if (FSoftObjectProperty* SoftObjectProperty = CastField<FSoftObjectProperty>(Property))
