Add Logs analysis toolset

#Synced-CL 52159831

diff --git a/Engine/Plugins/Experimental/ToolsetRegistry/Source/ToolsetRegistry/Private/Tests/LogsToolsetTest.cpp b/Engine/Plugins/Experimental/ToolsetRegistry/Source/ToolsetRegistry/Private/Tests/LogsToolsetTest.cpp
new file mode 100644
--- /dev/null
+++ b/Engine/Plugins/Experimental/ToolsetRegistry/Source/ToolsetRegistry/Private/Tests/LogsToolsetTest.cpp
@@ -0,0 +1,237 @@
+// Copyright Epic Games, Inc. All Rights Reserved.
+
+#include "CoreMinimal.h"
+#include "Misc/AutomationTest.h"
+#include "Misc/Guid.h"
+
+#include "Tests/ToolsetRegistryTestFlags.h"
+#include "ToolsetRegistry/LogsToolset.h"
+#include "ToolsetRegistry/ObjectFunctionToolCall.h"
+#include "ToolsetRegistry/ToolCallScriptExceptionHandler.h"
+
+#if WITH_DEV_AUTOMATION_TESTS
+
+DEFINE_LOG_CATEGORY_STATIC(LogLogsToolsetTest, Log, All);
+
+BEGIN_DEFINE_SPEC(FLogsToolsetSpec, "AI.ToolsetRegistry.LogsToolsetSpec",
+	ToolsetRegistryTest::Flags)
+	FString CallAndCaptureException(const FString& FunctionName, const FString& ParamsJson);
+END_DEFINE_SPEC(FLogsToolsetSpec)
+
+FString FLogsToolsetSpec::CallAndCaptureException(
+	const FString& FunctionName, const FString& ParamsJson)
+{
+	ULogsToolset* Toolset = NewObject<ULogsToolset>();
+	UFunction* Function = Toolset->GetClass()->FindFunctionByName(*FunctionName);
+	if (!TestNotNull("Function", Function)) return {};
+	auto ToolCall = MakeShared<UE::ToolsetRegistry::FObjectFunctionToolCall>(Toolset, Function);
+	auto Handler = MakeShared<UE::ToolsetRegistry::Internal::FToolCallScriptExceptionHandler>();
+	ToolCall->Execute(
+		UE::ToolsetRegistry::FObjectFunctionToolCall::FFunctionInputParamsJson(
+			TInPlaceType<FString>(), ParamsJson),
+		Handler);
+	return Handler->GetException();
+}
+
+void FLogsToolsetSpec::Define()
+{
+	Describe(TEXT("GetLogCategories"), [this]()
+	{
+		It(TEXT("Returns a non-empty sorted list"), [this]()
+		{
+			UE_LOG(LogLogsToolsetTest, Log, TEXT("GetLogCategories probe"));
+			GLog->Flush();
+
+			TArray<FString> Categories = ULogsToolset::GetLogCategories();
+			TestTrue("Categories is non-empty", Categories.Num() > 0);
+
+			TArray<FString> Sorted = Categories;
+			Sorted.Sort();
+			TestEqual("Categories is sorted", Categories, Sorted);
+		});
+
+		It(TEXT("Includes the test log category"), [this]()
+		{
+			UE_LOG(LogLogsToolsetTest, Log, TEXT("Category presence probe"));
+			GLog->Flush();
+
+			TArray<FString> Categories = ULogsToolset::GetLogCategories();
+			TestTrue("LogLogsToolsetTest is present",
+				Categories.Contains(TEXT("LogLogsToolsetTest")));
+		});
+
+		It(TEXT("Filter narrows results to matching categories"), [this]()
+		{
+			UE_LOG(LogLogsToolsetTest, Log, TEXT("Filter probe"));
+			GLog->Flush();
+
+			TArray<FString> Filtered = ULogsToolset::GetLogCategories(TEXT("LogLogsToolsetTest"));
+			TestTrue("At least one result", Filtered.Num() > 0);
+			for (const FString& Category : Filtered)
+			{
+				TestTrue("Category contains filter string",
+					Category.Contains(TEXT("LogLogsToolsetTest"), ESearchCase::IgnoreCase));
+			}
+		});
+	});
+
+	Describe(TEXT("GetLogEntries"), [this]()
+	{
+		It(TEXT("Returns a non-empty list with no filters"), [this]()
+		{
+			TArray<FString> Entries = ULogsToolset::GetLogEntries(TEXT(""), TEXT(""), 0);
+			TestTrue("Entries is non-empty", Entries.Num() > 0);
+		});
+
+		It(TEXT("Filters by category"), [this]()
+		{
+			UE_LOG(LogLogsToolsetTest, Log, TEXT("Category filter probe"));
+			GLog->Flush();
+
+			TArray<FString> Entries = ULogsToolset::GetLogEntries(
+				TEXT("LogLogsToolsetTest"), TEXT(""), 0);
+			TestTrue("At least one entry", Entries.Num() > 0);
+			for (const FString& Entry : Entries)
+			{
+				TestTrue("Entry contains category prefix",
+					Entry.Contains(TEXT("]LogLogsToolsetTest:")));
+			}
+		});
+
+		It(TEXT("Filters by pattern"), [this]()
+		{
+			const FString Marker = FString::Printf(
+				TEXT("LogsToolsetPatternMarker_%s"), *FGuid::NewGuid().ToString());
+			UE_LOG(LogLogsToolsetTest, Log, TEXT("%s"), *Marker);
+			GLog->Flush();
+
+			TArray<FString> Entries = ULogsToolset::GetLogEntries(TEXT(""), *Marker, 0);
+			if (TestEqual("Exactly one entry matches", Entries.Num(), 1))
+			{
+				TestTrue("Entry contains marker", Entries[0].Contains(Marker));
+			}
+		});
+
+		It(TEXT("Pattern matching is case-insensitive"), [this]()
+		{
+			const FString Marker = FString::Printf(
+				TEXT("LogsToolsetCaseMarker_%s"), *FGuid::NewGuid().ToString());
+			UE_LOG(LogLogsToolsetTest, Log, TEXT("%s"), *Marker);
+			GLog->Flush();
+
+			TArray<FString> Entries = ULogsToolset::GetLogEntries(
+				TEXT(""), *Marker.ToLower(), 0);
+			TestEqual("Case-insensitive match finds entry", Entries.Num(), 1);
+		});
+
+		It(TEXT("Raises an error for an unknown category"), [this]()
+		{
+			TestTrue("Script error raised",
+				CallAndCaptureException(
+					TEXT("GetLogEntries"),
+					TEXT(R"({"Category": "LogThisCategoryDoesNotExist_XYZ123", "Pattern": "", "MaxEntries": 0})"))
+				.Contains(TEXT("not found")));
+		});
+
+		It(TEXT("Respects MaxEntries limit"), [this]()
+		{
+			constexpr int32 Limit = 5;
+			TArray<FString> All = ULogsToolset::GetLogEntries(TEXT(""), TEXT(""), 0);
+			TArray<FString> Limited = ULogsToolset::GetLogEntries(TEXT(""), TEXT(""), Limit);
+			TestTrue("At most MaxEntries returned", Limited.Num() <= Limit);
+			if (All.Num() >= Limit)
+			{
+				TestEqual("Returns the most recent entries",
+					Limited, TArray<FString>(All.GetData() + All.Num() - Limit, Limit));
+			}
+		});
+
+		It(TEXT("Combines category and pattern filters"), [this]()
+		{
+			const FString Marker = FString::Printf(
+				TEXT("LogsToolsetCombinedMarker_%s"), *FGuid::NewGuid().ToString());
+			UE_LOG(LogLogsToolsetTest, Log, TEXT("%s"), *Marker);
+			GLog->Flush();
+
+			TArray<FString> Entries = ULogsToolset::GetLogEntries(
+				TEXT("LogLogsToolsetTest"), *Marker, 0);
+			TestEqual("Exactly one combined match", Entries.Num(), 1);
+		});
+	});
+
+	Describe(TEXT("GetVerbosity"), [this]()
+	{
+		It(TEXT("Returns a verbosity string for a known category"), [this]()
+		{
+			UE_LOG(LogLogsToolsetTest, Log, TEXT("GetVerbosity probe"));
+			const FString Verbosity = ULogsToolset::GetVerbosity(TEXT("LogLogsToolsetTest"));
+			TestFalse("Verbosity is non-empty", Verbosity.IsEmpty());
+		});
+
+		It(TEXT("Returns empty string for an unknown category"), [this]()
+		{
+			const FString Verbosity =
+				ULogsToolset::GetVerbosity(TEXT("LogThisCategoryDoesNotExist_XYZ123"));
+			TestTrue("Returns empty string for unknown category", Verbosity.IsEmpty());
+		});
+
+		It(TEXT("Raises an error for an unknown category"), [this]()
+		{
+			TestTrue("Script error raised",
+				CallAndCaptureException(
+					TEXT("GetVerbosity"),
+					TEXT(R"({"Category": "LogThisCategoryDoesNotExist_XYZ123"})"))
+				.Contains(TEXT("not found")));
+		});
+
+		It(TEXT("Raises an error for an empty category"), [this]()
+		{
+			TestTrue("Script error raised",
+				CallAndCaptureException(
+					TEXT("GetVerbosity"),
+					TEXT(R"({"Category": ""})"))
+				.Contains(TEXT("Category cannot be empty")));
+		});
+	});
+
+	Describe(TEXT("SetVerbosity"), [this]()
+	{
+		It(TEXT("Can round-trip a verbosity change"), [this]()
+		{
+			const FString Category = TEXT("LogLogsToolsetTest");
+			const FString Original = ULogsToolset::GetVerbosity(Category);
+			if (!TestFalse("Category verbosity is readable", Original.IsEmpty()))
+			{
+				return;
+			}
+
+			ULogsToolset::SetVerbosity(Category, TEXT("Warning"));
+			TestEqual("Verbosity set to Warning",
+				ULogsToolset::GetVerbosity(Category), FString(TEXT("Warning")));
+
+			ULogsToolset::SetVerbosity(Category, Original);
+			TestEqual("Verbosity restored",
+				ULogsToolset::GetVerbosity(Category), Original);
+		});
+
+		It(TEXT("Raises an error for an empty category"), [this]()
+		{
+			TestTrue("Script error raised",
+				CallAndCaptureException(
+					TEXT("SetVerbosity"),
+					TEXT(R"({"Category": "", "Verbosity": "Log"})"))
+				.Contains(TEXT("Category cannot be empty")));
+		});
+
+		It(TEXT("Raises an error for an invalid verbosity"), [this]()
+		{
+			TestTrue("Script error raised",
+				CallAndCaptureException(
+					TEXT("SetVerbosity"),
+					TEXT(R"({"Category": "LogLogsToolsetTest", "Verbosity": "NotAVerbosity"})"))
+				.Contains(TEXT("is not a valid verbosity level")));
+		});
+	});
+}
+
+#endif  // WITH_DEV_AUTOMATION_TESTS
diff --git a/Engine/Plugins/Experimental/ToolsetRegistry/Source/ToolsetRegistry/Private/ToolsetRegistry/LogsToolset.cpp b/Engine/Plugins/Experimental/ToolsetRegistry/Source/ToolsetRegistry/Private/ToolsetRegistry/LogsToolset.cpp
new file mode 100644
--- /dev/null
+++ b/Engine/Plugins/Experimental/ToolsetRegistry/Source/ToolsetRegistry/Private/ToolsetRegistry/LogsToolset.cpp
@@ -0,0 +1,236 @@
+// Copyright Epic Games, Inc. All Rights Reserved.
+
+#include "LogsToolset.h"
+
+#include "HAL/FileManager.h"
+#include "HAL/PlatformOutputDevices.h"
+#include "Internationalization/Regex.h"
+#include "Kismet/KismetSystemLibrary.h"
+#include "Misc/CoreMisc.h"
+#include "Misc/FileHelper.h"
+#include "Logging/LogVerbosity.h"
+#include "Misc/OutputDeviceNull.h"
+#include "Misc/Paths.h"
+
+#include UE_INLINE_GENERATED_CPP_BY_NAME(LogsToolset)
+
+namespace
+{
+	// Matches a standard UE log line prefix and captures the category name.
+	// Example: [2024.01.15-10.30.00:000][  0]LogTemp: Warning: Some message
+	// NOTE: FRegexPattern must NOT be a static variable — it requires i18n to be initialized.
+	const FString CategoryPatternStr = TEXT(R"regex(^\[[\d.\-:]+\]\[\s*\d+\](\w+):)regex");
+
+	// Valid verbosity names accepted by the log console command, derived directly
+	// from ELogVerbosity so the list stays in sync if new levels are ever added.
+	const TArray<FString> ValidVerbosities = []()
+	{
+		TArray<FString> Result;
+		for (int32 i = ELogVerbosity::NoLogging; i < ELogVerbosity::NumVerbosity; ++i)
+		{
+			Result.Add(ToString(static_cast<ELogVerbosity::Type>(i)));
+		}
+		return Result;
+	}();
+
+	// Output device that captures lines written to it into an array.
+	//
+	// GetLogCategories, GetVerbosity, and SetVerbosity all route through
+	// FSelfRegisteringExec::StaticExec with the "log" command. This is the same
+	// code path used by the editor's Output Log UI and UE's own developer tooling.
+	// It is the only public way to reach FLogSuppressionImplementation::ReverseAssociations,
+	// the private map that owns the live FLogCategoryBase* registry. FLogSuppressionInterface
+	// deliberately exposes only registration (AssociateSuppress / DisassociateSuppress),
+	// not enumeration or name-based lookup.
+	struct FLineCapture : public FOutputDevice
+	{
+		TArray<FString> Lines;
+
+		virtual void Serialize(
+			const TCHAR* V, ELogVerbosity::Type, const FName&) override
+		{
+			Lines.Add(FString(V));
+		}
+	};
+}
+
+TArray<FString> ULogsToolset::GetLogEntries(
+	const FString& Category, const FString& Pattern, int32 MaxEntries)
+{
+	if (!Category.IsEmpty())
+	{
+		const TArray<FString> Known = GetLogCategories(Category);
+		const bool bExists = Known.ContainsByPredicate([&Category](const FString& C)
+		{
+			return C.Equals(Category, ESearchCase::IgnoreCase);
+		});
+		if (!bExists)
+		{
+			UKismetSystemLibrary::RaiseScriptError(
+				FString::Printf(TEXT("Log category '%s' not found."), *Category));
+			return {};
+		}
+	}
+
+	const FRegexPattern CategoryRegex(CategoryPatternStr);
+	TOptional<FRegexPattern> PatternRegex;
+	if (!Pattern.IsEmpty())
+	{
+		PatternRegex.Emplace(Pattern, ERegexPatternFlags::CaseInsensitive);
+	}
+
+	TArray<FString> Results;
+
+	for (const FString& Line : ReadLogLines())
+	{
+		if (!Category.IsEmpty())
+		{
+			FRegexMatcher Matcher(CategoryRegex, Line);
+			if (!Matcher.FindNext() ||
+				!Matcher.GetCaptureGroup(1).Equals(Category, ESearchCase::IgnoreCase))
+			{
+				continue;
+			}
+		}
+
+		if (PatternRegex.IsSet())
+		{
+			FRegexMatcher Matcher(*PatternRegex, Line);
+			if (!Matcher.FindNext())
+			{
+				continue;
+			}
+		}
+
+		Results.Add(Line);
+	}
+
+	if (MaxEntries > 0 && Results.Num() > MaxEntries)
+	{
+		Results.RemoveAt(0, Results.Num() - MaxEntries);
+	}
+
+	return Results;
+}
+
+TArray<FString> ULogsToolset::GetLogCategories(const FString& Filter)
+{
+	// "log list" asks FLogSuppressionImplementation to enumerate every FLogCategoryBase
+	// that has registered itself via FLogSuppressionInterface::AssociateSuppress. Each
+	// output line is formatted as "%-40s  %-12s  %s" (name, verbosity, debug-break flag).
+	// Splitting on whitespace and taking the first token gives the category name.
+	// "log list <filter>" narrows results to categories whose name contains the filter
+	// string (case-insensitive).
+	FLineCapture Capture;
+	const FString Command = Filter.IsEmpty()
+		? TEXT("log list")
+		: FString::Printf(TEXT("log list %s"), *Filter);
+	FSelfRegisteringExec::StaticExec(nullptr, *Command, Capture);
+
+	TArray<FString> Categories;
+	for (const FString& Line : Capture.Lines)
+	{
+		TArray<FString> Parts;
+		Line.TrimEnd().ParseIntoArrayWS(Parts);
+		if (Parts.Num() > 0)
+		{
+			Categories.Add(Parts[0]);
+		}
+	}
+
+	Categories.Sort();
+	return Categories;
+}
+
+FString ULogsToolset::GetVerbosity(const FString& Category)
+{
+	if (Category.IsEmpty())
+	{
+		UKismetSystemLibrary::RaiseScriptError(TEXT("Category cannot be empty."));
+		return FString();
+	}
+
+	// "log list <filter>" filters by substring, so it may return multiple lines.
+	// Pass the exact name as the filter to keep the output small, then match exactly.
+	FLineCapture Capture;
+	FSelfRegisteringExec::StaticExec(
+		nullptr, *FString::Printf(TEXT("log list %s"), *Category), Capture);
+
+	for (const FString& Line : Capture.Lines)
+	{
+		TArray<FString> Parts;
+		Line.TrimEnd().ParseIntoArrayWS(Parts);
+		if (Parts.Num() >= 2 && Parts[0].Equals(Category, ESearchCase::IgnoreCase))
+		{
+			return Parts[1];
+		}
+	}
+
+	UKismetSystemLibrary::RaiseScriptError(
+		FString::Printf(TEXT("Log category '%s' not found."), *Category));
+	return FString();
+}
+
+void ULogsToolset::SetVerbosity(const FString& Category, const FString& Verbosity)
+{
+	if (Category.IsEmpty())
+	{
+		UKismetSystemLibrary::RaiseScriptError(TEXT("Category cannot be empty."));
+		return;
+	}
+
+	const bool bValid = ValidVerbosities.ContainsByPredicate(
+		[&Verbosity](const FString& V)
+		{
+			return V.Equals(Verbosity, ESearchCase::IgnoreCase);
+		});
+
+	if (!bValid)
+	{
+		UKismetSystemLibrary::RaiseScriptError(FString::Printf(
+			TEXT("'%s' is not a valid verbosity level. Use one of: %s"),
+			*Verbosity, *FString::Join(ValidVerbosities, TEXT(", "))));
+		return;
+	}
+
+	FOutputDeviceNull Null;
+	FSelfRegisteringExec::StaticExec(
+		nullptr, *FString::Printf(TEXT("log %s %s"), *Category, *Verbosity), Null);
+}
+
+FString ULogsToolset::GetLogFilePath()
+{
+	const FString AbsLogPath = FPaths::ConvertRelativePathToFull(
+		FPlatformOutputDevices::GetAbsoluteLogFilename());
+	if (!AbsLogPath.IsEmpty() && IFileManager::Get().FileExists(*AbsLogPath))
+	{
+		return AbsLogPath;
+	}
+
+	UKismetSystemLibrary::RaiseScriptError(
+		FString::Printf(TEXT("Log file not found: %s"), *AbsLogPath));
+	return FString();
+}
+
+TArray<FString> ULogsToolset::ReadLogLines()
+{
+	const FString LogFilePath = GetLogFilePath();
+	if (LogFilePath.IsEmpty())
+	{
+		return {};
+	}
+
+	FString Content;
+	if (!FFileHelper::LoadFileToString(Content, *LogFilePath,
+		FFileHelper::EHashOptions::None, FILEREAD_AllowWrite))
+	{
+		UKismetSystemLibrary::RaiseScriptError(
+			FString::Printf(TEXT("Failed to read log file: %s"), *LogFilePath));
+		return {};
+	}
+
+	TArray<FString> Lines;
+	Content.ParseIntoArrayLines(Lines, /*bCullEmpty=*/false);
+	Lines.RemoveAll([](const FString& Line) { return Line.TrimEnd().IsEmpty(); });
+	return Lines;
+}
diff --git a/Engine/Plugins/Experimental/ToolsetRegistry/Source/ToolsetRegistry/Private/ToolsetRegistry/LogsToolset.h b/Engine/Plugins/Experimental/ToolsetRegistry/Source/ToolsetRegistry/Private/ToolsetRegistry/LogsToolset.h
new file mode 100644
--- /dev/null
+++ b/Engine/Plugins/Experimental/ToolsetRegistry/Source/ToolsetRegistry/Private/ToolsetRegistry/LogsToolset.h
@@ -0,0 +1,65 @@
+// Copyright Epic Games, Inc. All Rights Reserved.
+
+#pragma once
+
+#include "CoreMinimal.h"
+
+#include "ToolsetRegistry/UToolsetRegistry.h"
+
+#include "LogsToolset.generated.h"
+
+/// Provides tools for reading the Unreal Engine output log and controlling
+/// log category verbosity.
+UCLASS(BlueprintType, MinimalAPI)
+class ULogsToolset : public UToolsetDefinition
+{
+	GENERATED_BODY()
+
+public:
+	/**
+	 * Returns log entries from the current session's log file.
+	 * @param Category If non-empty, only returns entries from this log category (e.g. "LogTemp").
+	 * @param Pattern If non-empty, only returns entries whose text matches this regular
+	 *   expression.
+	 * @param MaxEntries Maximum number of entries to return, taken from the end of the log.
+	 *   Pass 0 for no limit. Defaults to 1000.
+	 * @return A list of matching log entries in chronological order.
+	 */
+	UFUNCTION(BlueprintCallable, Category = "LogsToolset")
+	static TArray<FString> GetLogEntries(
+		const FString& Category = TEXT(""), const FString& Pattern = TEXT(""), int32 MaxEntries = 1000);
+
+	/**
+	 * Returns a sorted list of registered log categories.
+	 * @param Filter If non-empty, only returns categories whose name contains this
+	 *   substring.
+	 * @return A sorted list of log category names (e.g. ["LogBlueprint", "LogTemp"]).
+	 */
+	UFUNCTION(BlueprintCallable, Category = "LogsToolset")
+	static TArray<FString> GetLogCategories(const FString& Filter = TEXT(""));
+
+	/**
+	 * Returns the current verbosity level for a log category.
+	 * @param Category The log category name, e.g. "LogTemp".
+	 * @return The verbosity level as a string: one of "NoLogging", "Fatal", "Error",
+	 *   "Warning", "Display", "Log", "Verbose", or "VeryVerbose". Raises a script error
+	 *   if the category is not found.
+	 */
+	UFUNCTION(BlueprintCallable, Category = "LogsToolset")
+	static FString GetVerbosity(const FString& Category);
+
+	/**
+	 * Sets the verbosity level for a log category.
+	 * @param Category The log category name, e.g. "LogTemp".
+	 * @param Verbosity The verbosity level: one of "NoLogging", "Fatal", "Error",
+	 *   "Warning", "Display", "Log", "Verbose", or "VeryVerbose".
+	 */
+	UFUNCTION(BlueprintCallable, Category = "LogsToolset")
+	static void SetVerbosity(const FString& Category, const FString& Verbosity);
+
+private:
+	static FString GetLogFilePath();
+	static TArray<FString> ReadLogLines();
+
+	friend class FLogsToolsetSpec;
+};
diff --git a/Engine/Plugins/Experimental/ToolsetRegistry/Source/ToolsetRegistry/Private/ToolsetRegistry/ToolsetRegistrySubsystem.cpp b/Engine/Plugins/Experimental/ToolsetRegistry/Source/ToolsetRegistry/Private/ToolsetRegistry/ToolsetRegistrySubsystem.cpp
--- a/Engine/Plugins/Experimental/ToolsetRegistry/Source/ToolsetRegistry/Private/ToolsetRegistry/ToolsetRegistrySubsystem.cpp
+++ b/Engine/Plugins/Experimental/ToolsetRegistry/Source/ToolsetRegistry/Private/ToolsetRegistry/ToolsetRegistrySubsystem.cpp
@@ -7,6 +7,7 @@
 
 #include "ToolsetRegistry/AgentSkill.h"
 #include "ToolsetRegistry/EditorAppToolset.h"
+#include "ToolsetRegistry/LogsToolset.h"
 #include "ToolsetRegistry/Module.h"
 #include "ToolsetRegistry/UToolsetRegistry.h"
 
@@ -44,11 +45,13 @@ void UToolsetRegistrySubsystem::Initialize(FSubsystemCollectionBase& Collection)
 	Super::Initialize(Collection);
 	UToolsetRegistry::RegisterToolsetClass(UEditorAppToolset::StaticClass());
 	UToolsetRegistry::RegisterToolsetClass(UAgentSkillToolset::StaticClass());
+	UToolsetRegistry::RegisterToolsetClass(ULogsToolset::StaticClass());
 }
 
 void UToolsetRegistrySubsystem::Deinitialize()
 {
+	UToolsetRegistry::UnregisterToolsetClass(ULogsToolset::StaticClass());
 	UToolsetRegistry::UnregisterToolsetClass(UAgentSkillToolset::StaticClass());
 	UToolsetRegistry::UnregisterToolsetClass(UEditorAppToolset::StaticClass());
 	Super::Deinitialize();
-}
\ No newline at end of file
+}
