Moved selection and viewport specific logic from scene and asset toolsets to into EditorApp toolset.

Includes following improvements and fixes:
* Actor selection works in PIE mode
* Asset selection queries all content browsers.
* Asset selection is async to match editor behavior.
* Remove abuse of exec command in Focus on Actors.
* Added GetActorsInView.
* Closed gaps in test coverage for above functions.

#Synced-CL 52186709

diff --git a/Engine/Plugins/Experimental/ToolsetRegistry/Source/ToolsetRegistry/Private/Tests/EditorAppToolsetTest.cpp b/Engine/Plugins/Experimental/ToolsetRegistry/Source/ToolsetRegistry/Private/Tests/EditorAppToolsetTest.cpp
--- a/Engine/Plugins/Experimental/ToolsetRegistry/Source/ToolsetRegistry/Private/Tests/EditorAppToolsetTest.cpp
+++ b/Engine/Plugins/Experimental/ToolsetRegistry/Source/ToolsetRegistry/Private/Tests/EditorAppToolsetTest.cpp
@@ -3,13 +3,19 @@
 #include "CoreMinimal.h"
 #include "Misc/AutomationTest.h"
 
+#include "AssetRegistry/IAssetRegistry.h"
 #include "Editor.h"
 #include "Engine/World.h"
+#include "GameFramework/Actor.h"
 #include "Misc/App.h"
 #include "Misc/CoreMiscDefines.h"
+#include "Subsystems/UnrealEditorSubsystem.h"
 #include "UObject/StrongObjectPtr.h"
 #include "ToolsetRegistry/EditorAppToolset.h"
+#include "ToolsetRegistry/ToolCallAsyncResultVoid.h"
 #include "ToolsetRegistry/JsonConversion.h"
+#include "ToolsetRegistry/ObjectFunctionToolCall.h"
+#include "ToolsetRegistry/ToolCallScriptExceptionHandler.h"
 #include "ToolsetRegistry/ToolCallAsyncResultImage.h"
 
 
@@ -232,6 +238,194 @@ void FEditorAppToolsetSpec::Define()
 					return true; // done
 				}));
 		});
+
+	// Actor selection tests
+
+	It(TEXT("GetSelectedActors and SelectActors round-trip"), [this]()
+	{
+		UWorld* World = GEditor ? GEditor->GetEditorWorldContext().World() : nullptr;
+		if (!World)
+		{
+			return;
+		}
+
+		TArray<AActor*> OriginalSelection = UEditorAppToolset::GetSelectedActors();
+
+		AActor* TempActor = World->SpawnActor<AActor>(AActor::StaticClass());
+		if (!TestNotNull("TempActor", TempActor)) return;
+
+		ON_SCOPE_EXIT
+		{
+			UEditorAppToolset::SelectActors(OriginalSelection);
+			TempActor->Destroy();
+		};
+
+		UEditorAppToolset::SelectActors({TempActor});
+		TestTrue("TempActor is selected",
+			UEditorAppToolset::GetSelectedActors().Contains(TempActor));
+	});
+
+	// Viewport camera tests
+
+	It(TEXT("GetCameraTransform and SetCameraTransform round-trip"), [this]()
+	{
+		UUnrealEditorSubsystem* Subsystem = GEditor
+			? GEditor->GetEditorSubsystem<UUnrealEditorSubsystem>() : nullptr;
+		if (!Subsystem)
+		{
+			return;
+		}
+		FVector OriginalLocation;
+		FRotator OriginalRotation;
+		if (!Subsystem->GetLevelViewportCameraInfo(OriginalLocation, OriginalRotation))
+		{
+			// No viewport available; skip.
+			return;
+		}
+
+		const FVector TestLocation(100.f, 200.f, 300.f);
+		const FRotator TestRotation(10.f, 20.f, 0.f);
+		UEditorAppToolset::SetCameraTransform(FTransform(TestRotation, TestLocation));
+
+		FTransform Result = UEditorAppToolset::GetCameraTransform();
+		TestEqual("Camera location X", Result.GetLocation().X, TestLocation.X);
+		TestEqual("Camera location Y", Result.GetLocation().Y, TestLocation.Y);
+		TestEqual("Camera location Z", Result.GetLocation().Z, TestLocation.Z);
+		TestEqual("Camera rotation pitch", Result.Rotator().Pitch, TestRotation.Pitch);
+		TestEqual("Camera rotation yaw", Result.Rotator().Yaw, TestRotation.Yaw);
+
+		// Restore original camera position.
+		Subsystem->SetLevelViewportCameraInfo(OriginalLocation, OriginalRotation);
+	});
+
+	It(TEXT("FocusOnActors makes actor visible to GetVisibleActors"), [this]()
+	{
+		if (!GEditor || GEditor->PlayWorld)
+		{
+			return;
+		}
+		UWorld* World = GEditor->GetEditorWorldContext().World();
+		if (!World)
+		{
+			return;
+		}
+		UUnrealEditorSubsystem* Subsystem =
+			GEditor->GetEditorSubsystem<UUnrealEditorSubsystem>();
+		if (!Subsystem)
+		{
+			return;
+		}
+		FVector OriginalLocation;
+		FRotator OriginalRotation;
+		if (!Subsystem->GetLevelViewportCameraInfo(OriginalLocation, OriginalRotation))
+		{
+			// No viewport available; skip.
+			return;
+		}
+
+		AActor* TempActor = World->SpawnActor<AActor>(AActor::StaticClass());
+		if (!TestNotNull("TempActor", TempActor)) return;
+
+		UEditorAppToolset::FocusOnActors({TempActor});
+
+		AddCommand(new FUntilCommand(
+			[Subsystem, TempActor, OriginalLocation, OriginalRotation]() -> bool
+			{
+				if (!UEditorAppToolset::GetVisibleActors().Contains(TempActor)) return false;
+				Subsystem->SetLevelViewportCameraInfo(OriginalLocation, OriginalRotation);
+				TempActor->Destroy();
+				return true;
+			},
+			[this, Subsystem, TempActor, OriginalLocation, OriginalRotation]() -> bool
+			{
+				AddError(TEXT("Timed out waiting for actor to become visible after FocusOnActors."));
+				Subsystem->SetLevelViewportCameraInfo(OriginalLocation, OriginalRotation);
+				TempActor->Destroy();
+				return true;
+			}));
+	});
+
+	// Content browser tests
+
+	It(TEXT("GetContentBrowserPath and SetContentBrowserPath round-trip"), [this]()
+	{
+		FString OriginalPath = UEditorAppToolset::GetContentBrowserPath();
+
+		ON_SCOPE_EXIT
+		{
+			if (!OriginalPath.IsEmpty())
+			{
+				UEditorAppToolset::SetContentBrowserPath(OriginalPath);
+			}
+		};
+
+		UEditorAppToolset::SetContentBrowserPath(TEXT("/Game"));
+		TestEqual("Path is set", UEditorAppToolset::GetContentBrowserPath(), FString(TEXT("/Game")));
+	});
+
+	It(TEXT("SetContentBrowserPath raises a script error for an invalid path"), [this]()
+	{
+		UEditorAppToolset* Toolset = NewObject<UEditorAppToolset>();
+		UFunction* Function = Toolset->GetClass()->FindFunctionByName(TEXT("SetContentBrowserPath"));
+		if (!TestNotNull("Function", Function)) return;
+		auto ToolCall = MakeShared<UE::ToolsetRegistry::FObjectFunctionToolCall>(Toolset, Function);
+		auto Handler = MakeShared<UE::ToolsetRegistry::Internal::FToolCallScriptExceptionHandler>();
+		ToolCall->Execute(
+			UE::ToolsetRegistry::FObjectFunctionToolCall::FFunctionInputParamsJson(
+				TInPlaceType<FString>(),
+				TEXT(R"({"path": "/InvalidPath/DoesNotExist"})")),
+			Handler);
+		TestTrue("Script error raised", Handler->GetException().Contains(TEXT("Failed to navigate")));
+	});
+
+	It(TEXT("GetSelectedAssets and SelectAssets round-trip"), [this]()
+	{
+		TArray<FAssetData> GameAssets;
+		IAssetRegistry::GetChecked().GetAssetsByPath(
+			TEXT("/Game"), GameAssets, /*bRecursive=*/true);
+		if (GameAssets.IsEmpty())
+		{
+			// No game assets available to select; skip.
+			return;
+		}
+		const FString TestAsset = GameAssets[0].PackageName.ToString();
+		TArray<FString> OriginalSelection = UEditorAppToolset::GetSelectedAssets();
+
+		UToolCallAsyncResultVoid* Result = UEditorAppToolset::SelectAssets({TestAsset});
+		if (!TestNotNull("Result", Result)) return;
+
+		TStrongObjectPtr<UToolCallAsyncResultVoid> StrongResult(Result);
+		AddCommand(new FFunctionLatentCommand(
+			[this, StrongResult, TestAsset, OriginalSelection]() mutable -> bool
+			{
+				if (!StrongResult->bIsComplete) return false;
+				TestTrue("No error", StrongResult->Error.IsEmpty());
+				TestTrue("Asset is selected",
+					UEditorAppToolset::GetSelectedAssets().Contains(TestAsset));
+				UEditorAppToolset::SelectAssets(OriginalSelection);
+				StrongResult.Reset();
+				return true;
+			}));
+	});
+
+	It(TEXT("OpenEditorForAsset raises a script error for a non-existent asset"), [this]()
+	{
+		if (FApp::IsUnattended())
+		{
+			return;
+		}
+		UEditorAppToolset* Toolset = NewObject<UEditorAppToolset>();
+		UFunction* Function = Toolset->GetClass()->FindFunctionByName(TEXT("OpenEditorForAsset"));
+		if (!TestNotNull("Function", Function)) return;
+		auto ToolCall = MakeShared<UE::ToolsetRegistry::FObjectFunctionToolCall>(Toolset, Function);
+		auto Handler = MakeShared<UE::ToolsetRegistry::Internal::FToolCallScriptExceptionHandler>();
+		ToolCall->Execute(
+			UE::ToolsetRegistry::FObjectFunctionToolCall::FFunctionInputParamsJson(
+				TInPlaceType<FString>(),
+				TEXT(R"({"assetPath": "/Game/NonExistent/DoesNotExist"})")),
+			Handler);
+		TestTrue("Script error raised", Handler->GetException().Contains(TEXT("is not a valid asset path")));
+	});
 }
 
 #endif
diff --git a/Engine/Plugins/Experimental/ToolsetRegistry/Source/ToolsetRegistry/Private/ToolsetRegistry/EditorAppToolset.cpp b/Engine/Plugins/Experimental/ToolsetRegistry/Source/ToolsetRegistry/Private/ToolsetRegistry/EditorAppToolset.cpp
--- a/Engine/Plugins/Experimental/ToolsetRegistry/Source/ToolsetRegistry/Private/ToolsetRegistry/EditorAppToolset.cpp
+++ b/Engine/Plugins/Experimental/ToolsetRegistry/Source/ToolsetRegistry/Private/ToolsetRegistry/EditorAppToolset.cpp
@@ -2,16 +2,21 @@
 
 #include "EditorAppToolset.h"
 
+#include "Animation/AnimationAsset.h"
+#include "Animation/Skeleton.h"
 #include "AssetRegistry/IAssetRegistry.h"
+#include "ContentBrowserModule.h"
 #include "Containers/Ticker.h"
 #include "Editor.h"
+#include "EditorAssetLibrary.h"
+#include "IContentBrowserSingleton.h"
 #include "Engine/AssetManager.h"
-#include "Engine/StreamableManager.h"
-#include "Animation/AnimationAsset.h"
-#include "Animation/Skeleton.h"
 #include "Engine/SkeletalMesh.h"
 #include "Engine/StaticMesh.h"
+#include "Engine/StreamableManager.h"
 #include "Engine/Texture.h"
+#include "Engine/World.h"
+#include "EngineUtils.h"
 #include "HAL/ConsoleManager.h"
 #include "Kismet/KismetSystemLibrary.h"
 #include "LevelEditor.h"
@@ -23,12 +28,15 @@
 #include "Misc/PackageName.h"
 #include "ObjectTools.h"
 #include "RenderingThread.h"
-#include "ShaderCompiler.h"
+#include "SceneView.h"
 #include "SLevelViewport.h"
+#include "Selection.h"
+#include "Subsystems/UnrealEditorSubsystem.h"
 #include "Templates/SharedPointer.h"
 #include "ToolsetRegistry/DelegateHandle.h"
 #include "ToolsetRegistry/JsonConversion.h"
 #include "ToolsetRegistry/ToolCallAsyncResultImage.h"
+#include "ToolsetRegistry/ToolCallAsyncResultVoid.h"
 #include "ToolsetRegistry/ToolsetImage.h"
 #include "UObject/StrongObjectPtr.h"
 #include "UObject/UObjectGlobals.h"
@@ -395,3 +403,314 @@ UToolCallAsyncResultImage* UEditorAppToolset::CaptureAssetImage(const FString& A
 	}
 	return Result.Get();
 }
+
+TArray<AActor*> UEditorAppToolset::GetSelectedActors()
+{
+	TArray<AActor*> Result;
+	if (!GEditor)
+	{
+		UKismetSystemLibrary::RaiseScriptError(TEXT("Editor not found."));
+		return Result;
+	}
+	for (FSelectionIterator Iter(*GEditor->GetSelectedActors()); Iter; ++Iter)
+	{
+		if (AActor* Actor = Cast<AActor>(*Iter))
+		{
+			Result.Add(Actor);
+		}
+	}
+	return Result;
+}
+
+void UEditorAppToolset::SelectActors(const TArray<AActor*>& Actors)
+{
+	if (!GEditor)
+	{
+		UKismetSystemLibrary::RaiseScriptError(TEXT("Editor not found."));
+		return;
+	}
+	USelection* Selection = GEditor->GetSelectedActors();
+	Selection->Modify();
+	Selection->BeginBatchSelectOperation();
+	GEditor->SelectNone(false, true, false);
+	for (AActor* Actor : Actors)
+	{
+		if (Actor)
+		{
+			GEditor->SelectActor(Actor, true, false);
+		}
+	}
+	Selection->EndBatchSelectOperation();
+	GEditor->NoteSelectionChange();
+}
+
+FTransform UEditorAppToolset::GetCameraTransform()
+{
+	if (!GCurrentLevelEditingViewportClient)
+	{
+		UKismetSystemLibrary::RaiseScriptError(TEXT("There is no current level camera."));
+		return FTransform::Identity;
+	}
+	return FTransform(
+		GCurrentLevelEditingViewportClient->GetViewRotation(),
+		GCurrentLevelEditingViewportClient->GetViewLocation());
+}
+
+void UEditorAppToolset::SetCameraTransform(const FTransform& Transform)
+{
+	if (!GCurrentLevelEditingViewportClient)
+	{
+		UKismetSystemLibrary::RaiseScriptError(TEXT("There is no current level camera."));
+		return;
+	}
+	GCurrentLevelEditingViewportClient->SetViewLocationForOrbiting(Transform.GetLocation());
+	GCurrentLevelEditingViewportClient->SetViewLocation(Transform.GetLocation());
+	GCurrentLevelEditingViewportClient->SetViewRotation(Transform.Rotator());
+}
+
+void UEditorAppToolset::FocusOnActors(const TArray<AActor*>& Actors)
+{
+	if (!GEditor)
+	{
+		UKismetSystemLibrary::RaiseScriptError(TEXT("Editor not found."));
+		return;
+	}
+	if (GEditor->PlayWorld)
+	{
+		UKismetSystemLibrary::RaiseScriptError(TEXT("FocusOnActors is not supported during PIE."));
+		return;
+	}
+
+	GEditor->MoveViewportCamerasToActor(Actors, true);
+}
+
+TArray<AActor*> UEditorAppToolset::GetVisibleActors()
+{
+	TArray<AActor*> Result;
+	if (!GEditor || !GCurrentLevelEditingViewportClient || !GCurrentLevelEditingViewportClient->Viewport)
+	{
+		UKismetSystemLibrary::RaiseScriptError(TEXT("Editor not found."));
+		return Result;
+	}
+	UWorld* World = GEditor->GetEditorWorldContext().World();
+	if (!World)
+	{
+		UKismetSystemLibrary::RaiseScriptError(TEXT("No world to query."));
+		return Result;
+	}
+	FLevelEditorViewportClient* LevelVC = GCurrentLevelEditingViewportClient;
+	FSceneViewFamily ViewFamily(FSceneViewFamily::ConstructionValues(
+		LevelVC->Viewport, LevelVC->GetScene(), LevelVC->EngineShowFlags)
+		.SetRealtimeUpdate(LevelVC->IsRealtime()));
+	const TSharedPtr<FSceneView> SceneView = MakeShareable<>(LevelVC->CalcSceneView(&ViewFamily));
+	for (TActorIterator<AActor> It(World); It; ++It)
+	{
+		const FBox Box = It->GetComponentsBoundingBox(/*bNonColliding=*/true);
+		const bool bInFrustum = Box.IsValid
+			? SceneView->ViewFrustum.IntersectBox(Box.GetCenter(), Box.GetExtent())
+			: SceneView->ViewFrustum.IntersectSphere(It->GetActorLocation(), 0.f);
+		if (bInFrustum)
+		{
+			Result.Add(*It);
+		}
+	}
+	return Result;
+}
+
+FVector2D UEditorAppToolset::WorldPosToScreenCoords(FVector Position)
+{
+	if (!GEditor)
+	{
+		UKismetSystemLibrary::RaiseScriptError(TEXT("Editor not found."));
+		return FVector2D::ZeroVector;
+	}
+	UUnrealEditorSubsystem* Subsystem = GEditor->GetEditorSubsystem<UUnrealEditorSubsystem>();
+	check(Subsystem);
+	FVector2D ScreenCoords;
+	if (!Subsystem->WorldToScreen(Position, ScreenCoords))
+	{
+		UKismetSystemLibrary::RaiseScriptError(TEXT("Unable to convert position into screen space."));
+		return FVector2D::ZeroVector;
+	}
+	FIntPoint Size;
+	if (!Subsystem->GetLevelViewportSize(Size))
+	{
+		UKismetSystemLibrary::RaiseScriptError(TEXT("Unable to retrieve viewport size."));
+		return FVector2D::ZeroVector;
+	}
+	return FVector2D(ScreenCoords.X / Size.X, ScreenCoords.Y / Size.Y);
+}
+
+FVector UEditorAppToolset::ScreenCoordsToWorld(FVector2D Coords, float TraceDistance)
+{
+	if (Coords.X < 0 || Coords.X > 1 || Coords.Y < 0 || Coords.Y > 1)
+	{
+		UKismetSystemLibrary::RaiseScriptError(FString::Printf(
+			TEXT("Invalid viewport coords (%f, %f)."), Coords.X, Coords.Y));
+		return FVector::ZeroVector;
+	}
+	if (!GEditor)
+	{
+		UKismetSystemLibrary::RaiseScriptError(TEXT("Editor not found."));
+		return FVector::ZeroVector;
+	}
+	UWorld* World = GEditor->GetEditorWorldContext().World();
+	if (!World)
+	{
+		UKismetSystemLibrary::RaiseScriptError(TEXT("No world to query."));
+		return FVector::ZeroVector;
+	}
+	UUnrealEditorSubsystem* Subsystem = GEditor->GetEditorSubsystem<UUnrealEditorSubsystem>();
+	check(Subsystem);
+	FIntPoint Size;
+	if (!Subsystem->GetLevelViewportSize(Size))
+	{
+		UKismetSystemLibrary::RaiseScriptError(TEXT("Unable to retrieve viewport size."));
+		return FVector::ZeroVector;
+	}
+	FVector Origin, Direction;
+	const FVector2D PixelCoords(Coords.X * Size.X, Coords.Y * Size.Y);
+	if (!Subsystem->ScreenToWorld(PixelCoords, Origin, Direction))
+	{
+		UKismetSystemLibrary::RaiseScriptError(TEXT("Unable to convert screen coordinates into world space."));
+		return FVector::ZeroVector;
+	}
+	FHitResult HitResult;
+	TArray<AActor*> ActorsToIgnore;
+	const bool bHit = UKismetSystemLibrary::LineTraceSingle(
+		World,
+		Origin,
+		Origin + Direction * TraceDistance,
+		ETraceTypeQuery::TraceTypeQuery1,
+		true,
+		ActorsToIgnore,
+		EDrawDebugTrace::None,
+		HitResult,
+		true);
+	if (!bHit)
+	{
+		UKismetSystemLibrary::RaiseScriptError(TEXT("Unable to find a solid surface at the given screen coordinates."));
+		return FVector::ZeroVector;
+	}
+	return HitResult.Location;
+}
+
+TArray<FString> UEditorAppToolset::GetSelectedAssets()
+{
+	FContentBrowserModule& ContentBrowserModule =
+		FModuleManager::LoadModuleChecked<FContentBrowserModule>("ContentBrowser");
+	TArray<FAssetData> AssetData;
+	ContentBrowserModule.Get().GetAllSelectedAssets(AssetData);
+	TArray<FString> Result;
+	for (const FAssetData& Data : AssetData)
+	{
+		Result.Add(Data.PackageName.ToString());
+	}
+	return Result;
+}
+
+UToolCallAsyncResultVoid* UEditorAppToolset::SelectAssets(const TArray<FString>& AssetPaths)
+{
+	UToolCallAsyncResultVoid* Result = NewObject<UToolCallAsyncResultVoid>();
+
+	IAssetRegistry& AssetRegistry = IAssetRegistry::GetChecked();
+	TSet<FName> ExpectedPackages;
+	for (const FString& Path : AssetPaths)
+	{
+		TArray<FAssetData> Found;
+		AssetRegistry.GetAssetsByPackageName(FName(*Path), Found);
+		if (Found.Num() > 0)
+		{
+			ExpectedPackages.Add(Found[0].PackageName);
+		}
+	}
+
+	UEditorAssetLibrary::SyncBrowserToObjects(AssetPaths);
+
+	// Selection is applied asynchronously by the content browser. Poll each tick
+	// until all expected assets appear in the selection or the timeout elapses.
+	TStrongObjectPtr<UToolCallAsyncResultVoid> StrongResult(Result);
+	FTSTicker::GetCoreTicker().AddTicker(FTickerDelegate::CreateLambda(
+		[StrongResult, ExpectedPackages, StartTime = FPlatformTime::Seconds()](float) mutable -> bool
+		{
+			FContentBrowserModule& CBModule =
+				FModuleManager::LoadModuleChecked<FContentBrowserModule>("ContentBrowser");
+			TArray<FAssetData> SelectedAssets;
+			CBModule.Get().GetSelectedAssets(SelectedAssets);
+			TSet<FName> SelectedPackages;
+			for (const FAssetData& Data : SelectedAssets)
+			{
+				SelectedPackages.Add(Data.PackageName);
+			}
+			bool bAllSelected = true;
+			for (const FName& Expected : ExpectedPackages)
+			{
+				if (!SelectedPackages.Contains(Expected))
+				{
+					bAllSelected = false;
+					break;
+				}
+			}
+			if (bAllSelected)
+			{
+				StrongResult->SetCompleted();
+				StrongResult.Reset();
+				return false;
+			}
+			const float MaxSelectionTime = 5.f;
+			if (FPlatformTime::Seconds() - StartTime > MaxSelectionTime)
+			{
+				StrongResult->SetError(TEXT("Timed out waiting for asset selection."));
+				StrongResult.Reset();
+				return false;
+			}
+			return true; // keep polling
+		}));
+
+	return Result;
+}
+
+FString UEditorAppToolset::GetContentBrowserPath()
+{
+	FContentBrowserModule& ContentBrowserModule =
+		FModuleManager::LoadModuleChecked<FContentBrowserModule>("ContentBrowser");
+	const FContentBrowserItemPath CurrentPath = ContentBrowserModule.Get().GetCurrentPath();
+	return CurrentPath.HasInternalPath() ? CurrentPath.GetInternalPathString() : FString();
+}
+
+void UEditorAppToolset::SetContentBrowserPath(const FString& Path)
+{
+	FContentBrowserModule& ContentBrowserModule =
+		FModuleManager::LoadModuleChecked<FContentBrowserModule>("ContentBrowser");
+	ContentBrowserModule.Get().SetSelectedPaths({Path});
+	const FString ResultPath = UEditorAppToolset::GetContentBrowserPath();
+	if (ResultPath != Path)
+	{
+		UKismetSystemLibrary::RaiseScriptError(FString::Printf(
+			TEXT("Failed to navigate content browser to '%s': ended up at '%s'."),
+			*Path, *ResultPath));
+	}
+}
+
+void UEditorAppToolset::OpenEditorForAsset(const FString& AssetPath)
+{
+	if (!GEditor)
+	{
+		UKismetSystemLibrary::RaiseScriptError(TEXT("Editor not found."));
+		return;
+	}
+	UObject* Asset = StaticLoadObject(UObject::StaticClass(), nullptr, *AssetPath);
+	if (!Asset)
+	{
+		UKismetSystemLibrary::RaiseScriptError(
+			FString::Printf(TEXT("%s is not a valid asset path."), *AssetPath));
+		return;
+	}
+	UAssetEditorSubsystem* Subsystem = GEditor->GetEditorSubsystem<UAssetEditorSubsystem>();
+	check(Subsystem);
+	if (!Subsystem->OpenEditorForAsset(Asset))
+	{
+		UKismetSystemLibrary::RaiseScriptError(
+			FString::Printf(TEXT("Failed to open editor for asset: %s"), *AssetPath));
+	}
+}
diff --git a/Engine/Plugins/Experimental/ToolsetRegistry/Source/ToolsetRegistry/Private/ToolsetRegistry/EditorAppToolset.h b/Engine/Plugins/Experimental/ToolsetRegistry/Source/ToolsetRegistry/Private/ToolsetRegistry/EditorAppToolset.h
--- a/Engine/Plugins/Experimental/ToolsetRegistry/Source/ToolsetRegistry/Private/ToolsetRegistry/EditorAppToolset.h
+++ b/Engine/Plugins/Experimental/ToolsetRegistry/Source/ToolsetRegistry/Private/ToolsetRegistry/EditorAppToolset.h
@@ -8,12 +8,15 @@
 
 #include "EditorAppToolset.generated.h"
 
+class AActor;
 class FJsonObject;
 class IConsoleObject;
 class UObject;
 class UToolCallAsyncResultImage;
+class UToolCallAsyncResultVoid;
 
-/// Contains tools that query or modify the state of the Unreal Editor application.
+/// Tools for querying and modifying Unreal Editor state: console variables, asset imaging,
+/// actor selection, viewport camera, and content browser navigation.
 UCLASS(BlueprintType, MinimalAPI)
 class UEditorAppToolset : public UToolsetDefinition
 {
@@ -38,6 +41,103 @@ public:
 	UFUNCTION(BlueprintCallable, Category = "EditorAppToolset")
 	static UToolCallAsyncResultImage* CaptureAssetImage(const FString& AssetPath);
 
+	/*
+	 * Gets the currently selected actors in the level editor.
+	 * @return A list of the currently selected actors.
+	 */
+	UFUNCTION(BlueprintCallable, Category = "EditorAppToolset")
+	static TArray<AActor*> GetSelectedActors();
+
+	/*
+	 * Selects the specified actors in the current scene.
+	 * @param Actors The actors to select.
+	 */
+	UFUNCTION(BlueprintCallable, Category = "EditorAppToolset")
+	static void SelectActors(const TArray<AActor*>& Actors);
+
+	/*
+	 * Returns the position and rotation of the level viewport camera.
+	 * @return The transform for the level viewport camera.
+	 */
+	UFUNCTION(BlueprintCallable, Category = "EditorAppToolset")
+	static FTransform GetCameraTransform();
+
+	/*
+	 * Sets the position and rotation of the level viewport camera.
+	 * @param Transform The transform to apply to the viewport camera.
+	 */
+	UFUNCTION(BlueprintCallable, Category = "EditorAppToolset")
+	static void SetCameraTransform(const FTransform& Transform);
+
+	/*
+	 * Repositions the level editor camera to focus on the specified actors.
+	 * Cannot be called while PIE is active.
+	 * @param Actors The actors to focus the level camera on.
+	 */
+	UFUNCTION(BlueprintCallable, Category = "EditorAppToolset")
+	static void FocusOnActors(const TArray<AActor*>& Actors);
+
+	/*
+	 * Returns all actors in the current level whose bounds intersect the viewport frustum.
+	 * @return The actors that are at least partially in view.
+	 */
+	UFUNCTION(BlueprintCallable, Category = "EditorAppToolset")
+	static TArray<AActor*> GetVisibleActors();
+
+	/*
+	 * Converts a world-space position into normalized screen space based on the editor viewport camera.
+	 * @param Position The world space position to convert.
+	 * @return The normalized screen-space coordinates for the position if it is in view.
+	 */
+	UFUNCTION(BlueprintCallable, Category = "EditorAppToolset")
+	static FVector2D WorldPosToScreenCoords(FVector Position);
+
+	/*
+	 * Finds the world position of the nearest solid object at a given set of normalized view space coords.
+	 * @param Coords The normalized screen-space coordinates to trace from.
+	 * @param TraceDistance The maximum distance to trace within the scene.
+	 * @return The world-space coordinates of the intersection point.
+	 */
+	UFUNCTION(BlueprintCallable, Category = "EditorAppToolset")
+	static FVector ScreenCoordsToWorld(FVector2D Coords, float TraceDistance = 100000.f);
+
+	/*
+	 * Gets the list of assets selected in the content browser.
+	 * @return A list of package paths to the assets selected in the content browser.
+	 */
+	UFUNCTION(BlueprintCallable, Category = "EditorAppToolset")
+	static TArray<FString> GetSelectedAssets();
+
+	/*
+	 * Selects the specified assets in the content browser.
+	 * Completes once the content browser has applied the selection.
+	 * @param AssetPaths The package paths of the assets to select.
+	 * @return Async result that completes when the selection is applied.
+	 */
+	UFUNCTION(BlueprintCallable, Category = "EditorAppToolset")
+	static UToolCallAsyncResultVoid* SelectAssets(const TArray<FString>& AssetPaths);
+
+	/*
+	 * Gets the current path of the active content browser.
+	 * @return The path to the folder currently shown in the content browser.
+	 */
+	UFUNCTION(BlueprintCallable, Category = "EditorAppToolset")
+	static FString GetContentBrowserPath();
+
+	/*
+	 * Navigates the active content browser to the specified folder path.
+	 * @param Path The internal path to navigate to, e.g. '/Game/Meshes'.
+	 */
+	UFUNCTION(BlueprintCallable, Category = "EditorAppToolset")
+	static void SetContentBrowserPath(const FString& Path);
+
+	/*
+	 * Opens an asset editor for the specified asset.
+	 * @param AssetPath The path to the asset to open, e.g. '/Game/Meshes/SM_Cube'.
+	 */
+	UFUNCTION(BlueprintCallable, Category = "EditorAppToolset")
+	static void OpenEditorForAsset(const FString& AssetPath);
+
 private:
 	static TSharedPtr<FJsonObject> CVarToJson(IConsoleObject* CVar);
 
diff --git a/Engine/Plugins/Experimental/ToolsetRegistry/Source/ToolsetRegistry/ToolsetRegistry.Build.cs b/Engine/Plugins/Experimental/ToolsetRegistry/Source/ToolsetRegistry/ToolsetRegistry.Build.cs
--- a/Engine/Plugins/Experimental/ToolsetRegistry/Source/ToolsetRegistry/ToolsetRegistry.Build.cs
+++ b/Engine/Plugins/Experimental/ToolsetRegistry/Source/ToolsetRegistry/ToolsetRegistry.Build.cs
@@ -39,6 +39,8 @@ public class ToolsetRegistry : ModuleRules
 			new string[]
 			{
 				"AssetRegistry",
+				"ContentBrowser",
+				"ContentBrowserData",
 				"BlueprintGraph",
 				"CoreUObject",
 				"DeveloperSettings",
