Tweak editor app docs.

#Synced-CL 52186709

diff --git a/Engine/Plugins/Experimental/ToolsetRegistry/Source/ToolsetRegistry/Private/ToolsetRegistry/EditorAppToolset.h b/Engine/Plugins/Experimental/ToolsetRegistry/Source/ToolsetRegistry/Private/ToolsetRegistry/EditorAppToolset.h
--- a/Engine/Plugins/Experimental/ToolsetRegistry/Source/ToolsetRegistry/Private/ToolsetRegistry/EditorAppToolset.h
+++ b/Engine/Plugins/Experimental/ToolsetRegistry/Source/ToolsetRegistry/Private/ToolsetRegistry/EditorAppToolset.h
@@ -16,7 +16,7 @@ class UToolCallAsyncResultImage;
 class UToolCallAsyncResultVoid;
 
 /// Tools for querying and modifying Unreal Editor state: console variables, asset imaging,
-/// actor selection, viewport camera, and content browser navigation.
+/// actor and asset selection, viewport camera, and content browser navigation.
 UCLASS(BlueprintType, MinimalAPI)
 class UEditorAppToolset : public UToolsetDefinition
 {
