[ToolsetRegistry] FSoftObjectPath bracket name resolution fix

Bypass ImportText_Direct for ObjectProperty and ClassProperty paths that contain
bracket characters (e.g. "[Overlay]"). ImportText_Direct uses
FPropertyHelpers::ReadToken which stops at '[' and ' ', truncating subobject paths.

Fix: resolve via FSoftObjectPath::ResolveObject() + SetObjectPropertyValue() directly.

#Synced-CL 52169528

diff --git a/Engine/Plugins/Experimental/ToolsetRegistry/Source/ToolsetRegistry/Private/Tests/ReferenceConverterTest.cpp b/Engine/Plugins/Experimental/ToolsetRegistry/Source/ToolsetRegistry/Private/Tests/ReferenceConverterTest.cpp
--- a/Engine/Plugins/Experimental/ToolsetRegistry/Source/ToolsetRegistry/Private/Tests/ReferenceConverterTest.cpp
+++ b/Engine/Plugins/Experimental/ToolsetRegistry/Source/ToolsetRegistry/Private/Tests/ReferenceConverterTest.cpp
@@ -483,6 +483,77 @@ void FToolsetReferenceConverterSpec::Define()
 		TestEqual("RefCheck", TestOut.TestObject, TestOther);
 	});
 
+	It("Converts bracket-named UObject* to and from JSON", [this]()
+	{
+		// Systems like UMG auto-generate bracket-prefixed names for non-variable widgets
+		// (e.g. "[Overlay] Camera Mode Content"). ImportText_Direct truncates at '['
+		// because FPropertyHelpers::ReadToken defines valid chars as a-zA-Z0-9_-+./:
+		// The ReferenceConverter fix bypasses ImportText_Direct by using
+		// FSoftObjectPath::ResolveObject() + SetObjectPropertyValue() directly.
+
+		UObject* Outer = NewObject<UReferenceConverterTestObject>(
+			GetTransientPackage(), TEXT("BracketNameTestOuter"));
+
+		UReferenceConverterTestObject* BracketObj = NewObject<UReferenceConverterTestObject>(
+			Outer, FName("[Overlay] Test Content"));
+
+		FReferenceConverterTest TestIn;
+		TestIn.TestObject = BracketObj;
+
+		UStruct* Struct = FReferenceConverterTest::StaticStruct();
+		FProperty* Property = Struct->FindPropertyByName("TestObject");
+
+		TSharedPtr<FJsonValue> OutJson = PropertyToJsonData(Property, &TestIn.TestObject);
+		TestObjectRefData(OutJson->AsObject(), BracketObj->GetPathName());
+
+		FReferenceConverterTest TestOut;
+		TestTrue("Bracket name round-trips through JSON",
+			JsonDataToProperty(OutJson, Property, &TestOut.TestObject));
+		TestEqual("Resolved to same object", TestIn.TestObject, TestOut.TestObject);
+	});
+
+	It("Resolves bracket-named object via ExecuteTool", [this]()
+	{
+		UObject* Outer = NewObject<UReferenceConverterTestObject>(
+			GetTransientPackage(), TEXT("BracketNameTestOuter2"));
+
+		UReferenceConverterTestObject* BracketObj = NewObject<UReferenceConverterTestObject>(
+			Outer, FName("[TestClass] My Content"));
+
+		TObjectPtr<UReferenceConverterTestObject> TestObject = NewObject<UReferenceConverterTestObject>();
+		UE::ToolsetRegistry::FJsonValueOrError Result =
+			UE::ToolsetRegistry::TestHelpers::ExecuteToolCallWithJson(
+				TestObject, TEXT("TestObjectParam"),
+				FString::Printf(TEXT(R"({"testObject": {"refPath": "%s"}})"),
+					*BracketObj->GetPathName()));
+
+		TestFalse("Bracket-named refPath should resolve without error", Result.HasError());
+	});
+
+	It("Resolves space-named object via JSON round-trip", [this]()
+	{
+		// Control test: spaces without brackets should also work
+		UObject* Outer = NewObject<UReferenceConverterTestObject>(
+			GetTransientPackage(), TEXT("SpaceNameTestOuter"));
+
+		UReferenceConverterTestObject* SpaceObj = NewObject<UReferenceConverterTestObject>(
+			Outer, FName("Camera Mode Content"));
+
+		FReferenceConverterTest TestIn;
+		TestIn.TestObject = SpaceObj;
+
+		UStruct* Struct = FReferenceConverterTest::StaticStruct();
+		FProperty* Property = Struct->FindPropertyByName("TestObject");
+
+		TSharedPtr<FJsonValue> OutJson = PropertyToJsonData(Property, &TestIn.TestObject);
+		TestObjectRefData(OutJson->AsObject(), SpaceObj->GetPathName());
+
+		FReferenceConverterTest TestOut;
+		TestTrue("Space name round-trips through JSON",
+			JsonDataToProperty(OutJson, Property, &TestOut.TestObject));
+		TestEqual("Resolved to same object", TestIn.TestObject, TestOut.TestObject);
+	});
+
 	It("Creates instanced sub-object from class refPath", [this]()
 	{
 		TObjectPtr<UReferenceConverterTestObject> Outer = NewObject<UReferenceConverterTestObject>();
diff --git a/Engine/Plugins/Experimental/ToolsetRegistry/Source/ToolsetRegistry/Private/ToolsetRegistry/ReferenceConverter.cpp b/Engine/Plugins/Experimental/ToolsetRegistry/Source/ToolsetRegistry/Private/ToolsetRegistry/ReferenceConverter.cpp
--- a/Engine/Plugins/Experimental/ToolsetRegistry/Source/ToolsetRegistry/Private/ToolsetRegistry/ReferenceConverter.cpp
+++ b/Engine/Plugins/Experimental/ToolsetRegistry/Source/ToolsetRegistry/Private/ToolsetRegistry/ReferenceConverter.cpp
@@ -288,6 +288,29 @@ namespace UE::ToolsetRegistry
 										*RefPath, *MetaClass->GetName(), *Property->GetName()));
 					return true;
 				}
+
+				// Set property value directly via FSoftObjectPath resolution.
+				// ImportText_Direct uses FPropertyHelpers::ReadToken which stops at '[' and ' '
+				// characters, truncating paths with bracket-named subobjects.
+				if (ClassProperty)
+				{
+					ClassProperty->SetObjectPropertyValue(OutValue, Cast<UClass>(Object));
+				}
+				else if (ObjectProperty)
+				{
+					ObjectProperty->SetObjectPropertyValue(OutValue, Object);
+				}
+				else
+				{
+					// FSoftClassProperty — fall through to ImportText_Direct
+					FStringOutputDevice ImportError;
+					Property->ImportText_Direct(*RefPath, OutValue, nullptr, 0, &ImportError);
+					if (!ImportError.IsEmpty())
+					{
+						UKismetSystemLibrary::RaiseScriptError(ImportError);
+					}
+				}
+				return true;
 			}
 			else if (ObjectProperty)
 			{
@@ -328,6 +351,11 @@ namespace UE::ToolsetRegistry
 								*RefPath, *ObjectProperty->PropertyClass->GetName(), *Property->GetName()));
 						return true;
 					}
+
+					// Set property value directly — bypasses ImportText_Direct's broken
+					// tokenizer which truncates paths at '[' bracket characters.
+					ObjectProperty->SetObjectPropertyValue(OutValue, Object);
+					return true;
 				}
 			}
 			// Validate soft references using Asset Registry
@@ -397,6 +425,7 @@ namespace UE::ToolsetRegistry
 			}
 		}
 
+		// Fallback for non-object properties (soft refs, struct paths, etc.)
 		FStringOutputDevice ImportError;
 		Property->ImportText_Direct(*RefPath, OutValue, nullptr, 0, &ImportError);
 		if (!ImportError.IsEmpty())
