[AI Assistant] Expose ToolCallExceptionHandler in public toolset registry API for toolset testing.

Added tests to ensure:
* Script exceptions / errors are caught in the lifetime of the object.
* Script exceptions / errors are concatenated in the exception string.
* Script exceptions / errors are ignored outside of the object's lifetime.

#Synced-CL 52269267

diff --git a/Engine/Plugins/Experimental/ToolsetRegistry/Source/ToolsetRegistry/Private/Tests/EditorAppToolsetTest.cpp b/Engine/Plugins/Experimental/ToolsetRegistry/Source/ToolsetRegistry/Private/Tests/EditorAppToolsetTest.cpp
--- a/Engine/Plugins/Experimental/ToolsetRegistry/Source/ToolsetRegistry/Private/Tests/EditorAppToolsetTest.cpp
+++ b/Engine/Plugins/Experimental/ToolsetRegistry/Source/ToolsetRegistry/Private/Tests/EditorAppToolsetTest.cpp
@@ -15,8 +15,8 @@
 #include "ToolsetRegistry/ToolCallAsyncResultVoid.h"
 #include "ToolsetRegistry/JsonConversion.h"
 #include "ToolsetRegistry/ObjectFunctionToolCall.h"
-#include "ToolsetRegistry/ToolCallScriptExceptionHandler.h"
 #include "ToolsetRegistry/ToolCallAsyncResultImage.h"
+#include "ToolsetRegistry/ToolCallExceptionHandler.h"
 
 
 #if WITH_DEV_AUTOMATION_TESTS
@@ -369,7 +369,7 @@ void FEditorAppToolsetSpec::Define()
 		UFunction* Function = Toolset->GetClass()->FindFunctionByName(TEXT("SetContentBrowserPath"));
 		if (!TestNotNull("Function", Function)) return;
 		auto ToolCall = MakeShared<UE::ToolsetRegistry::FObjectFunctionToolCall>(Toolset, Function);
-		auto Handler = MakeShared<UE::ToolsetRegistry::Internal::FToolCallScriptExceptionHandler>();
+		auto Handler = MakeShared<UE::ToolsetRegistry::FToolCallExceptionHandler>();
 		ToolCall->Execute(
 			UE::ToolsetRegistry::FObjectFunctionToolCall::FFunctionInputParamsJson(
 				TInPlaceType<FString>(),
@@ -418,7 +418,7 @@ void FEditorAppToolsetSpec::Define()
 		UFunction* Function = Toolset->GetClass()->FindFunctionByName(TEXT("OpenEditorForAsset"));
 		if (!TestNotNull("Function", Function)) return;
 		auto ToolCall = MakeShared<UE::ToolsetRegistry::FObjectFunctionToolCall>(Toolset, Function);
-		auto Handler = MakeShared<UE::ToolsetRegistry::Internal::FToolCallScriptExceptionHandler>();
+		auto Handler = MakeShared<UE::ToolsetRegistry::FToolCallExceptionHandler>();
 		ToolCall->Execute(
 			UE::ToolsetRegistry::FObjectFunctionToolCall::FFunctionInputParamsJson(
 				TInPlaceType<FString>(),
diff --git a/Engine/Plugins/Experimental/ToolsetRegistry/Source/ToolsetRegistry/Private/Tests/LogsToolsetTest.cpp b/Engine/Plugins/Experimental/ToolsetRegistry/Source/ToolsetRegistry/Private/Tests/LogsToolsetTest.cpp
--- a/Engine/Plugins/Experimental/ToolsetRegistry/Source/ToolsetRegistry/Private/Tests/LogsToolsetTest.cpp
+++ b/Engine/Plugins/Experimental/ToolsetRegistry/Source/ToolsetRegistry/Private/Tests/LogsToolsetTest.cpp
@@ -7,7 +7,7 @@
 #include "Tests/ToolsetRegistryTestFlags.h"
 #include "ToolsetRegistry/LogsToolset.h"
 #include "ToolsetRegistry/ObjectFunctionToolCall.h"
-#include "ToolsetRegistry/ToolCallScriptExceptionHandler.h"
+#include "ToolsetRegistry/ToolCallExceptionHandler.h"
 
 #if WITH_DEV_AUTOMATION_TESTS
 
@@ -25,7 +25,7 @@ FString FLogsToolsetSpec::CallAndCaptureException(
 	UFunction* Function = Toolset->GetClass()->FindFunctionByName(*FunctionName);
 	if (!TestNotNull("Function", Function)) return {};
 	auto ToolCall = MakeShared<UE::ToolsetRegistry::FObjectFunctionToolCall>(Toolset, Function);
-	auto Handler = MakeShared<UE::ToolsetRegistry::Internal::FToolCallScriptExceptionHandler>();
+	auto Handler = MakeShared<UE::ToolsetRegistry::FToolCallExceptionHandler>();
 	ToolCall->Execute(
 		UE::ToolsetRegistry::FObjectFunctionToolCall::FFunctionInputParamsJson(
 			TInPlaceType<FString>(), ParamsJson),
diff --git a/Engine/Plugins/Experimental/ToolsetRegistry/Source/ToolsetRegistry/Private/Tests/ToolCallExceptionHandlerTest.cpp b/Engine/Plugins/Experimental/ToolsetRegistry/Source/ToolsetRegistry/Private/Tests/ToolCallExceptionHandlerTest.cpp
new file mode 100644
--- /dev/null
+++ b/Engine/Plugins/Experimental/ToolsetRegistry/Source/ToolsetRegistry/Private/Tests/ToolCallExceptionHandlerTest.cpp
@@ -0,0 +1,237 @@
+// Copyright Epic Games, Inc. All Rights Reserved.
+
+#include "ToolCallExceptionHandlerTest.h"
+
+#include "Internationalization/Text.h"
+#include "Kismet/KismetSystemLibrary.h"
+#include "Misc/AssertionMacros.h"
+#include "Misc/AutomationTest.h"
+
+#include "Tests/ToolCallTestHelpers.h"
+#include "Tests/ToolsetRegistryTestFlags.h"
+#include "ToolsetRegistry/JsonValueOrError.h"
+#include "ToolsetRegistry/ObjectFunctionToolCall.h"
+#include "ToolsetRegistry/ToolCallExceptionHandler.h"
+#include "ToolsetRegistry/ToolCallExceptionHandlerInternal.h"
+
+PRAGMA_DISABLE_DEPRECATION_WARNINGS
+
+#if WITH_DEV_AUTOMATION_TESTS
+
+BEGIN_DEFINE_SPEC(FToolCallExceptionHandlerSpec,
+	"AI.ToolsetRegistry.ToolCallExceptionHandler",
+	ToolsetRegistryTest::Flags)
+
+	// Throw a blueprint exception without an error message.
+	static void ThrowScriptException(EBlueprintExceptionType::Type ExceptionType)
+	{
+		FFrame* TopFrame = FFrame::GetThreadLocalTopStackFrame();
+		check(TopFrame);
+		FBlueprintCoreDelegates::ThrowScriptException(
+			TopFrame->Object, *TopFrame,
+			FBlueprintExceptionInfo(ExceptionType, FText()));
+	}
+END_DEFINE_SPEC(FToolCallExceptionHandlerSpec)
+
+void FToolCallExceptionHandlerSpec::Define()
+{
+	using namespace UE::ToolsetRegistry;
+
+	Describe(TEXT("GetException"), [this]()
+	{
+		It(TEXT("Captures an error raised by a function"), [this]()
+		{
+			FString Error = TEXT("an error");
+			FToolCallExceptionHandler Handler;
+			Handler.CaptureErrorsIn(
+				[&Error]()
+				{
+					UKismetSystemLibrary::RaiseScriptError(Error);
+				});
+			TestEqual(TEXT("Error matches"), Handler.GetException(), Error);
+		});
+
+		It(TEXT("Accumulates multiple errors and separates them with newlines"), [this]()
+		{
+			FString Error1 = TEXT("an error");
+			FString Error2 = TEXT("another error");
+			FToolCallExceptionHandler Handler;
+			Handler.CaptureErrorsIn(
+				[&Error1, Error2]()
+				{
+					UKismetSystemLibrary::RaiseScriptError(Error1);
+					UKismetSystemLibrary::RaiseScriptError(Error2);
+				});
+			TestEqual(TEXT("Error matches"), Handler.GetException(), Error1 + TEXT("\n") + Error2);
+		});
+
+		It(TEXT("Reports errors when no error message is provided"), [this]()
+		{
+			for (const auto& ExceptionToErrorString :
+				UE::ToolsetRegistry::Internal::HandledBlueprintExceptionTypeToErrorStrings)
+			{
+				FToolCallExceptionHandler Handler;
+				Handler.CaptureErrorsIn(
+					[ExceptionType = ExceptionToErrorString.Key]()
+					{
+						ThrowScriptException(ExceptionType);
+					});
+				FString ExceptionMessage = Handler.GetException();
+				TestNotEqual(
+					*FString::Printf(TEXT("Error reported for '%s'"), *ExceptionToErrorString.Value),
+					*ExceptionMessage, TEXT(""));
+				TestTrue(
+					*FString::Printf(TEXT("Error contains '%s'"), *ExceptionToErrorString.Value),
+					ExceptionMessage.Contains(*ExceptionToErrorString.Value));
+				
+			}
+		});
+
+		It(TEXT("Doesn't report errors for ignored exceptions"), [this]()
+		{
+			for (const auto& ExceptionType :
+				UE::ToolsetRegistry::Internal::IgnoredBlueprintExceptionTypes)
+			{
+				FToolCallExceptionHandler Handler;
+				Handler.CaptureErrorsIn(
+					[ExceptionType]()
+					{
+						ThrowScriptException(ExceptionType);
+					});
+				TestEqual(TEXT("No error"), *Handler.GetException(), TEXT(""));
+			}
+		});
+
+		It(TEXT("Doesn't catch errors outside of the object lifetime"), [this]()
+		{
+			auto HandlerCatchesNoErrors = [this]()
+				{
+					FToolCallExceptionHandler Handler;
+					TestEqual(TEXT("No error"), *Handler.GetException(), TEXT(""));
+				};
+			HandlerCatchesNoErrors();
+			UE::ToolsetRegistry::Internal::CallWithinBlueprintScript(
+				[]()
+				{
+					UKismetSystemLibrary::RaiseScriptError();
+				});
+			HandlerCatchesNoErrors();
+		});
+	});
+
+	Describe(TEXT("ToolCalls"), [this]()
+	{
+		It(TEXT("Accumulates multiple parameter errors with newlines"), [this]()
+		{
+			TObjectPtr<UToolCallExceptionHandlerTestObject> TestObject =
+				NewObject<UToolCallExceptionHandlerTestObject>();
+
+			// Pass two invalid references to trigger multiple RaiseScriptError calls
+			FString InvalidJson = FString::Printf(
+				TEXT(R"({"firstObject": {"refPath": "%s"}, "secondObject": {"refPath": "%s"}})"),
+				*UObject::StaticClass()->GetPathName(),  // Wrong type for FirstObject
+				*UClass::StaticClass()->GetPathName());  // Wrong type for SecondObject
+
+			FJsonValueOrError Result = TestHelpers::ExecuteToolCallWithJson(
+				TestObject, TEXT("TestMultipleRefs"), InvalidJson);
+
+			if (!TestTrue(TEXT("Should have error"), Result.HasError())) return;
+
+			FString ErrorMessage = Result.GetError();
+
+			// Verify both errors are present
+			TestTrue(TEXT("Contains first param type validation error"),
+					 ErrorMessage.Contains(
+						 TEXT("is not valid ToolCallExceptionHandlerTestObject for "
+							 "property 'FirstObject'")));
+			TestTrue(TEXT("Contains second param type validation error"),
+					 ErrorMessage.Contains(
+						 TEXT("is not valid ToolCallExceptionHandlerTestObject for "
+							 "property 'SecondObject'")));
+
+			// Verify errors are accumulated with newlines
+			TestTrue(TEXT("Contains newline separator"),
+					 ErrorMessage.Contains(TEXT("\n")));
+
+			// Verify parameter error prefix
+			TestTrue(TEXT("Contains parameter error prefix"),
+					 ErrorMessage.Contains(TEXT("Parameter error")));
+		});
+
+		It(TEXT("Captures single parameter error correctly"), [this]()
+		{
+			TObjectPtr<UToolCallExceptionHandlerTestObject> TestObject =
+				NewObject<UToolCallExceptionHandlerTestObject>();
+
+			// Pass one invalid reference
+			FString InvalidJson = FString::Printf(
+				TEXT(R"({"testObject": {"refPath": "%s"}})"),
+				*UObject::StaticClass()->GetPathName());  // Wrong type
+
+			FJsonValueOrError Result = TestHelpers::ExecuteToolCallWithJson(
+				TestObject, TEXT("TestSpecObjectParam"), InvalidJson);
+
+			if (!TestTrue("Should have error", Result.HasError())) return;
+
+			FString ErrorMessage = Result.GetError();
+
+			// Verify error message format
+			TestTrue(TEXT("Contains type validation error"),
+					 ErrorMessage.Contains(
+						 TEXT("is not valid ToolCallExceptionHandlerTestObject")));
+			TestTrue(TEXT("Error has parameter prefix"),
+					 ErrorMessage.Contains(TEXT("Parameter error")));
+		});
+
+		It(TEXT("Returns no error when all parameters are valid"), [this]()
+		{
+			TObjectPtr<UToolCallExceptionHandlerTestObject> TestObject =
+				NewObject<UToolCallExceptionHandlerTestObject>();
+			TObjectPtr<UToolCallExceptionHandlerTestObject> ValidObject1 =
+				NewObject<UToolCallExceptionHandlerTestObject>();
+			TObjectPtr<UToolCallExceptionHandlerTestObject> ValidObject2 =
+				NewObject<UToolCallExceptionHandlerTestObject>();
+
+			// Pass valid references
+			FString ValidJson = FString::Printf(
+				TEXT(R"({"firstObject": {"refPath": "%s"}, "secondObject": {"refPath": "%s"}})"),
+				*ValidObject1->GetPathName(),
+				*ValidObject2->GetPathName());
+
+			FJsonValueOrError Result = TestHelpers::ExecuteToolCallWithJson(
+				TestObject, TEXT("TestMultipleRefs"), ValidJson);
+
+			// Should succeed with no error
+			TestFalse(TEXT("Should not have error"), Result.HasError());
+		});
+
+		It(TEXT("Returns a fallback message when an empty error message is raised"), [this]()
+		{
+			TObjectPtr<UToolCallExceptionHandlerTestObject> TestObject =
+				NewObject<UToolCallExceptionHandlerTestObject>();
+
+			UFunction* Function = TestObject->GetClass()->FindFunctionByName(TEXT("TestEmptyError"));
+			if (!TestNotNull(TEXT("Function exists"), Function)) return;
+
+			auto ToolCall = MakeShared<FObjectFunctionToolCall>(TestObject, Function);
+			TSharedPtr<FToolCallExceptionHandler> ExceptionHandler =
+				MakeShared<FToolCallExceptionHandler>();
+
+			ToolCall->Execute(
+				FObjectFunctionToolCall::FFunctionInputParamsJson(TInPlaceType<FString>(), FString()),
+				ExceptionHandler);
+
+			FString ErrorMessage = ExceptionHandler->GetException();
+
+			if (!TestFalse(TEXT("Should have error"), ErrorMessage.IsEmpty())) return;
+
+			// Verify error message format
+			TestTrue(TEXT("Error has fallback format"),
+				ErrorMessage.Contains(TEXT("UserRaisedError in ")));
+		});
+	});
+}
+
+#endif
+
+PRAGMA_ENABLE_DEPRECATION_WARNINGS
diff --git a/Engine/Plugins/Experimental/ToolsetRegistry/Source/ToolsetRegistry/Private/Tests/ToolCallExceptionHandlerTest.h b/Engine/Plugins/Experimental/ToolsetRegistry/Source/ToolsetRegistry/Private/Tests/ToolCallExceptionHandlerTest.h
new file mode 100644
--- /dev/null
+++ b/Engine/Plugins/Experimental/ToolsetRegistry/Source/ToolsetRegistry/Private/Tests/ToolCallExceptionHandlerTest.h
@@ -0,0 +1,30 @@
+// Copyright Epic Games, Inc. All Rights Reserved.
+
+#pragma once
+
+#include "CoreMinimal.h"
+#include "Kismet/KismetSystemLibrary.h"
+
+#include "ToolCallExceptionHandlerTest.generated.h"
+
+
+UCLASS(BlueprintType)
+class UToolCallExceptionHandlerTestObject : public UObject
+{
+	GENERATED_BODY()
+public:
+	UFUNCTION(BlueprintCallable, Category="ToolCallExceptionHandlerTest")
+	void TestSpecObjectParam(UToolCallExceptionHandlerTestObject* TestObject) {}
+
+	UFUNCTION(BlueprintCallable, Category="ToolCallExceptionHandlerTest")
+	void TestMultipleRefs(
+		UToolCallExceptionHandlerTestObject* FirstObject,
+		UToolCallExceptionHandlerTestObject* SecondObject) {}
+
+	UFUNCTION(BlueprintCallable, Category="ToolCallExceptionHandlerTest")
+	void TestEmptyError()
+	{
+		// Raise an explicitly empty error
+		UKismetSystemLibrary::RaiseScriptError(FString());
+	}
+};
diff --git a/Engine/Plugins/Experimental/ToolsetRegistry/Source/ToolsetRegistry/Private/Tests/ToolCallScriptExceptionHandlerTest.cpp b/Engine/Plugins/Experimental/ToolsetRegistry/Source/ToolsetRegistry/Private/Tests/ToolCallScriptExceptionHandlerTest.cpp
deleted file mode 100644
diff --git a/Engine/Plugins/Experimental/ToolsetRegistry/Source/ToolsetRegistry/Private/Tests/ToolCallScriptExceptionHandlerTest.h b/Engine/Plugins/Experimental/ToolsetRegistry/Source/ToolsetRegistry/Private/Tests/ToolCallScriptExceptionHandlerTest.h
deleted file mode 100644
diff --git a/Engine/Plugins/Experimental/ToolsetRegistry/Source/ToolsetRegistry/Private/Tests/ToolCallTestHelpers.h b/Engine/Plugins/Experimental/ToolsetRegistry/Source/ToolsetRegistry/Private/Tests/ToolCallTestHelpers.h
--- a/Engine/Plugins/Experimental/ToolsetRegistry/Source/ToolsetRegistry/Private/Tests/ToolCallTestHelpers.h
+++ b/Engine/Plugins/Experimental/ToolsetRegistry/Source/ToolsetRegistry/Private/Tests/ToolCallTestHelpers.h
@@ -6,7 +6,7 @@
 #include "Templates/ValueOrError.h"
 #include "ToolsetRegistry/JsonValueOrError.h"
 #include "ToolsetRegistry/ObjectFunctionToolCall.h"
-#include "ToolsetRegistry/ToolCallScriptExceptionHandler.h"
+#include "ToolsetRegistry/ToolCallExceptionHandler.h"
 
 namespace UE::ToolsetRegistry::TestHelpers
 {
@@ -25,8 +25,8 @@ namespace UE::ToolsetRegistry::TestHelpers
 		}
 
 		auto ToolCall = MakeShared<FObjectFunctionToolCall>(TestObject, Function);
-		TSharedPtr<Internal::FToolCallScriptExceptionHandler> ExceptionHandler =
-			MakeShared<Internal::FToolCallScriptExceptionHandler>();
+		TSharedPtr<FToolCallExceptionHandler> ExceptionHandler =
+			MakeShared<FToolCallExceptionHandler>();
 
 		TFuture<FJsonValueOrError> ResultFuture = ToolCall->Execute(
 			FObjectFunctionToolCall::FFunctionInputParamsJson(
diff --git a/Engine/Plugins/Experimental/ToolsetRegistry/Source/ToolsetRegistry/Private/ToolsetRegistry/BlueprintLibraryToolset.cpp b/Engine/Plugins/Experimental/ToolsetRegistry/Source/ToolsetRegistry/Private/ToolsetRegistry/BlueprintLibraryToolset.cpp
--- a/Engine/Plugins/Experimental/ToolsetRegistry/Source/ToolsetRegistry/Private/ToolsetRegistry/BlueprintLibraryToolset.cpp
+++ b/Engine/Plugins/Experimental/ToolsetRegistry/Source/ToolsetRegistry/Private/ToolsetRegistry/BlueprintLibraryToolset.cpp
@@ -116,7 +116,7 @@ namespace UE::ToolsetRegistry
 
 		// Capture script errors raised by the tool during execution. Note: script errors raised by
 		// tools running outside of the game thread will not be captured correctly.
-		auto ExceptionHandler = MakeShared<FToolCallScriptExceptionHandler>();
+		auto ExceptionHandler = MakeShared<FToolCallExceptionHandler>();
 		ensureAlways(IsInGameThread());
 
 		FString ToolNameCopy = ToolName;
diff --git a/Engine/Plugins/Experimental/ToolsetRegistry/Source/ToolsetRegistry/Private/ToolsetRegistry/ObjectFunctionToolCall.cpp b/Engine/Plugins/Experimental/ToolsetRegistry/Source/ToolsetRegistry/Private/ToolsetRegistry/ObjectFunctionToolCall.cpp
--- a/Engine/Plugins/Experimental/ToolsetRegistry/Source/ToolsetRegistry/Private/ToolsetRegistry/ObjectFunctionToolCall.cpp
+++ b/Engine/Plugins/Experimental/ToolsetRegistry/Source/ToolsetRegistry/Private/ToolsetRegistry/ObjectFunctionToolCall.cpp
@@ -283,7 +283,7 @@ namespace UE::ToolsetRegistry
 
 	TFuture<FJsonValueOrError> FObjectFunctionToolCall::Execute(
 		const TOptional<FFunctionInputParamsJson>& FunctionInputParamsJson,
-		TSharedPtr<FToolCallScriptExceptionHandler> ExceptionHandler) const
+		TSharedPtr<FToolCallExceptionHandler> ExceptionHandler) const
 	{
 		// Get function input params JSON as a JSON object.
 		TSharedPtr<FJsonObject> FunctionInputParamsJsonObject;
diff --git a/Engine/Plugins/Experimental/ToolsetRegistry/Source/ToolsetRegistry/Private/ToolsetRegistry/ObjectFunctionToolCall.h b/Engine/Plugins/Experimental/ToolsetRegistry/Source/ToolsetRegistry/Private/ToolsetRegistry/ObjectFunctionToolCall.h
--- a/Engine/Plugins/Experimental/ToolsetRegistry/Source/ToolsetRegistry/Private/ToolsetRegistry/ObjectFunctionToolCall.h
+++ b/Engine/Plugins/Experimental/ToolsetRegistry/Source/ToolsetRegistry/Private/ToolsetRegistry/ObjectFunctionToolCall.h
@@ -10,7 +10,7 @@
 #include "UObject/Class.h"
 
 #include "ToolsetRegistry/JsonValueOrError.h"
-#include "ToolsetRegistry/ToolCallScriptExceptionHandler.h"
+#include "ToolsetRegistry/ToolCallExceptionHandler.h"
 
 //
 // UE::ToolsetRegistry::FObjectFunctionToolCall
@@ -67,7 +67,7 @@ namespace UE::ToolsetRegistry
 		TFuture<FJsonValueOrError> Execute(
 			const TOptional<FFunctionInputParamsJson>& FunctionInputParamsJson = 
 			TOptional<FFunctionInputParamsJson>(),
-			TSharedPtr<Internal::FToolCallScriptExceptionHandler> ExceptionHandler = nullptr) const;
+			TSharedPtr<FToolCallExceptionHandler> ExceptionHandler = nullptr) const;
 
 	private:
 		FJsonValueOrError BuildValidFunctionInputParamsJsonObject(
diff --git a/Engine/Plugins/Experimental/ToolsetRegistry/Source/ToolsetRegistry/Private/ToolsetRegistry/ToolCallExceptionHandler.cpp b/Engine/Plugins/Experimental/ToolsetRegistry/Source/ToolsetRegistry/Private/ToolsetRegistry/ToolCallExceptionHandler.cpp
new file mode 100644
--- /dev/null
+++ b/Engine/Plugins/Experimental/ToolsetRegistry/Source/ToolsetRegistry/Private/ToolsetRegistry/ToolCallExceptionHandler.cpp
@@ -0,0 +1,108 @@
+// Copyright Epic Games, Inc. All Rights Reserved.
+
+#include "ToolsetRegistry/ToolCallExceptionHandler.h"
+
+#include "Blueprint/BlueprintExceptionInfo.h"
+#include "Templates/UnrealTemplate.h"
+#include "UObject/Stack.h"
+
+#include "ToolsetRegistry/DelegateHandle.h"
+#include "ToolsetRegistry/Module.h"
+#include "ToolsetRegistry/ToolCallExceptionHandlerInternal.h"
+
+namespace UE::ToolsetRegistry
+{
+	FToolCallExceptionHandler::FToolCallExceptionHandler()
+		: ScriptExceptionHandle(
+			FDelegateHandleRaii::Create(
+				FBlueprintCoreDelegates::OnScriptException,
+				FBlueprintCoreDelegates::OnScriptException.AddRaw(
+					this, &FToolCallExceptionHandler::OnScriptException))),
+		  RuntimeScriptExceptionHandle(
+			  [this](ELogVerbosity::Type Verbosity, const TCHAR* ExceptionMessage, const TCHAR* StackMessage)
+			  {
+				  AddToExceptionMessage(FString(ExceptionMessage));
+
+				  // Also log to console to make more obvious to tool developers.
+				  UE_LOG(LogToolsetRegistry, Warning, TEXT("Error running tool:\n%s"), ExceptionMessage);
+			  })
+	{}
+
+	void FToolCallExceptionHandler::OnScriptException(
+			const UObject* ActiveObject, const FFrame& StackFrame, const FBlueprintExceptionInfo& Info)
+	{
+		TOptional<FString> ErrorMessage = GetBlueprintExceptionMessage(Info, StackFrame);
+		if (ErrorMessage.IsSet())
+		{
+			AddToExceptionMessage(ErrorMessage.GetValue());
+		}
+	}
+
+	FString FToolCallExceptionHandler::GetBackupExceptionMessage(
+		const FString& ExceptionTypeName, const FFrame& StackFrame)
+	{
+		// Generate as useful a message as we can if the exception didn't provide one.
+		FString FallbackExceptionMessage;
+		if (StackFrame.PreviousTrackingFrame &&
+			StackFrame.PreviousTrackingFrame->Object &&
+			StackFrame.PreviousTrackingFrame->Node)
+		{
+			FallbackExceptionMessage = FString::Printf(
+				TEXT("%s in %s.%s"), *ExceptionTypeName,
+				*StackFrame.PreviousTrackingFrame->Object->GetName(),
+				*StackFrame.PreviousTrackingFrame->Node->GetName());
+		}
+		else if (StackFrame.Object && StackFrame.Node)
+		{
+			FallbackExceptionMessage = FString::Printf(
+				TEXT("%s in %s.%s"), *ExceptionTypeName,
+				*StackFrame.Object->GetName(),
+				*StackFrame.Node->GetName());
+		}
+		else
+		{
+			FallbackExceptionMessage = ExceptionTypeName;
+		}
+		return FallbackExceptionMessage;
+	}
+
+	TOptional<FString> FToolCallExceptionHandler::GetBlueprintExceptionMessage(
+		const FBlueprintExceptionInfo& Info, const FFrame& StackFrame)
+	{
+		const FText& Description = Info.GetDescription();
+		if (!Description.IsEmpty()) return Description.ToString();
+
+		EBlueprintExceptionType::Type ExceptionType = Info.GetType();
+		const FString* ErrorString =
+			UE::ToolsetRegistry::Internal::HandledBlueprintExceptionTypeToErrorStrings.Find(
+				ExceptionType);
+		if (ErrorString) return GetBackupExceptionMessage(*ErrorString, StackFrame);
+
+		checkfSlow(
+			UE::ToolsetRegistry::Internal::IgnoredBlueprintExceptionTypes.Contains(ExceptionType),
+			TEXT("Unhandled exception type %d"), ExceptionType);
+		return TOptional<FString>();
+	}
+
+	void FToolCallExceptionHandler::AddToExceptionMessage(const FString& NewMessage)
+	{
+		if (NewMessage.IsEmpty())
+		{
+			return;
+		}
+		else if (ScriptException.IsEmpty())
+		{
+			ScriptException = NewMessage;
+		}
+		else
+		{
+			ScriptException += FString::Printf(TEXT("\n%s"), *NewMessage);
+		}
+	}
+
+	void FToolCallExceptionHandler::CaptureErrorsIn(TFunction<void()>&& Func)
+	{
+		UE::ToolsetRegistry::Internal::CallWithinBlueprintScript(MoveTemp(Func));
+	}
+}
+
diff --git a/Engine/Plugins/Experimental/ToolsetRegistry/Source/ToolsetRegistry/Private/ToolsetRegistry/ToolCallExceptionHandlerInternal.cpp b/Engine/Plugins/Experimental/ToolsetRegistry/Source/ToolsetRegistry/Private/ToolsetRegistry/ToolCallExceptionHandlerInternal.cpp
new file mode 100644
--- /dev/null
+++ b/Engine/Plugins/Experimental/ToolsetRegistry/Source/ToolsetRegistry/Private/ToolsetRegistry/ToolCallExceptionHandlerInternal.cpp
@@ -0,0 +1,34 @@
+// Copyright Epic Games, Inc. All Rights Reserved.
+
+#include "ToolsetRegistry/ToolCallExceptionHandlerInternal.h"
+
+#include "Containers/UnrealString.h"
+#include "Templates/UnrealTemplate.h"
+
+#include "ToolsetRegistry/ToolCallExceptionHandlerTestScope.h"
+
+namespace UE::ToolsetRegistry::Internal
+{
+	const TMap<EBlueprintExceptionType::Type, FString>
+		HandledBlueprintExceptionTypeToErrorStrings =
+	{
+		{EBlueprintExceptionType::AccessViolation, TEXT("AccessViolation")},
+		{EBlueprintExceptionType::InfiniteLoop, TEXT("InfiniteLoop")},
+		{EBlueprintExceptionType::NonFatalError, TEXT("NonFatalError")},
+		{EBlueprintExceptionType::FatalError, TEXT("FatalError")},
+		{EBlueprintExceptionType::AbortExecution, TEXT("AbortExecution")},
+		{EBlueprintExceptionType::UserRaisedError, TEXT("UserRaisedError")},
+	};
+	
+	const TSet<EBlueprintExceptionType::Type> IgnoredBlueprintExceptionTypes =
+	{
+		EBlueprintExceptionType::Breakpoint,
+		EBlueprintExceptionType::Tracepoint,
+		EBlueprintExceptionType::WireTracepoint,
+	};
+
+	void CallWithinBlueprintScript(TFunction<void()>&& Func)
+	{
+		NewObject<UToolCallExceptionHandlerTestScope>()->CallFunction(MoveTemp(Func));
+	}
+}
\ No newline at end of file
diff --git a/Engine/Plugins/Experimental/ToolsetRegistry/Source/ToolsetRegistry/Private/ToolsetRegistry/ToolCallExceptionHandlerInternal.h b/Engine/Plugins/Experimental/ToolsetRegistry/Source/ToolsetRegistry/Private/ToolsetRegistry/ToolCallExceptionHandlerInternal.h
new file mode 100644
--- /dev/null
+++ b/Engine/Plugins/Experimental/ToolsetRegistry/Source/ToolsetRegistry/Private/ToolsetRegistry/ToolCallExceptionHandlerInternal.h
@@ -0,0 +1,22 @@
+// Copyright Epic Games, Inc. All Rights Reserved.
+
+#pragma once
+
+#include "CoreMinimal.h"
+
+#include "Blueprint/BlueprintExceptionInfo.h"
+#include "Containers/Map.h"
+#include "Containers/Set.h"
+#include "Templates/Function.h"
+
+namespace UE::ToolsetRegistry::Internal
+{
+	// Map of FToolCallExceptionHandler handled blueprint exception types to error strings.
+	extern const TMap<EBlueprintExceptionType::Type, FString>
+		HandledBlueprintExceptionTypeToErrorStrings;
+	// Set of FToolCallExceptionHandler ignored blueprint exception types to error strings.
+	extern const TSet<EBlueprintExceptionType::Type> IgnoredBlueprintExceptionTypes;
+
+	// Call a function within the scope of a blueprint script.
+	void CallWithinBlueprintScript(TFunction<void()>&& Func);
+}
\ No newline at end of file
diff --git a/Engine/Plugins/Experimental/ToolsetRegistry/Source/ToolsetRegistry/Private/ToolsetRegistry/ToolCallExceptionHandlerTestScope.cpp b/Engine/Plugins/Experimental/ToolsetRegistry/Source/ToolsetRegistry/Private/ToolsetRegistry/ToolCallExceptionHandlerTestScope.cpp
new file mode 100644
--- /dev/null
+++ b/Engine/Plugins/Experimental/ToolsetRegistry/Source/ToolsetRegistry/Private/ToolsetRegistry/ToolCallExceptionHandlerTestScope.cpp
@@ -0,0 +1,22 @@
+// Copyright Epic Games, Inc. All Rights Reserved.
+
+#include "ToolsetRegistry/ToolCallExceptionHandlerTestScope.h"
+
+#include "Templates/UnrealTemplate.h"
+#include "UObject/NameTypes.h"
+
+void UToolCallExceptionHandlerTestScope::CallFunction(TFunction<void()>&& Func)
+{
+	FunctionToCall = MoveTemp(Func);
+
+	UFunction* CallFunctionInternalObject =
+		FindFunctionChecked(FName(TEXT("CallFunctionInternal")));
+	ProcessEvent(CallFunctionInternalObject, nullptr);
+	
+	FunctionToCall.Reset();
+}
+
+void UToolCallExceptionHandlerTestScope::CallFunctionInternal()
+{
+	if (FunctionToCall.IsSet()) FunctionToCall();
+}
diff --git a/Engine/Plugins/Experimental/ToolsetRegistry/Source/ToolsetRegistry/Private/ToolsetRegistry/ToolCallExceptionHandlerTestScope.h b/Engine/Plugins/Experimental/ToolsetRegistry/Source/ToolsetRegistry/Private/ToolsetRegistry/ToolCallExceptionHandlerTestScope.h
new file mode 100644
--- /dev/null
+++ b/Engine/Plugins/Experimental/ToolsetRegistry/Source/ToolsetRegistry/Private/ToolsetRegistry/ToolCallExceptionHandlerTestScope.h
@@ -0,0 +1,32 @@
+// Copyright Epic Games, Inc. All Rights Reserved.
+
+#pragma once
+
+#include "CoreMinimal.h"
+
+#include "Templates/Function.h"
+#include "UObject/Object.h"
+#include "UObject/ObjectMacros.h"
+
+#include "ToolCallExceptionHandlerTestScope.generated.h"
+
+
+// Simplifies testing Blueprint toolsets that raise errors caught by FToolCallExceptionHandler.
+// NOTE: This is tested by ToolCallExceptionHandlerTest.cpp.
+UCLASS(BlueprintType, HideDropdown)
+class UToolCallExceptionHandlerTestScope : public UObject
+{
+	GENERATED_BODY()
+
+public:
+	// Call the specified function within a blueprint execution context.
+	void CallFunction(TFunction<void()>&& Func);
+
+protected:
+	// Calls FunctionToCall if set.
+	UFUNCTION()
+	void CallFunctionInternal();
+
+private:
+	TFunction<void()> FunctionToCall;
+};
\ No newline at end of file
diff --git a/Engine/Plugins/Experimental/ToolsetRegistry/Source/ToolsetRegistry/Private/ToolsetRegistry/ToolCallScriptExceptionHandler.cpp b/Engine/Plugins/Experimental/ToolsetRegistry/Source/ToolsetRegistry/Private/ToolsetRegistry/ToolCallScriptExceptionHandler.cpp
deleted file mode 100644
diff --git a/Engine/Plugins/Experimental/ToolsetRegistry/Source/ToolsetRegistry/Private/ToolsetRegistry/ToolCallScriptExceptionHandler.h b/Engine/Plugins/Experimental/ToolsetRegistry/Source/ToolsetRegistry/Private/ToolsetRegistry/ToolCallScriptExceptionHandler.h
deleted file mode 100644
diff --git a/Engine/Plugins/Experimental/ToolsetRegistry/Source/ToolsetRegistry/Public/ToolsetRegistry/ToolCallExceptionHandler.h b/Engine/Plugins/Experimental/ToolsetRegistry/Source/ToolsetRegistry/Public/ToolsetRegistry/ToolCallExceptionHandler.h
new file mode 100644
--- /dev/null
+++ b/Engine/Plugins/Experimental/ToolsetRegistry/Source/ToolsetRegistry/Public/ToolsetRegistry/ToolCallExceptionHandler.h
@@ -0,0 +1,85 @@
+// Copyright Epic Games, Inc. All Rights Reserved.
+
+#pragma once
+
+#include "Blueprint/BlueprintExceptionInfo.h"
+#include "Containers/Array.h"
+#include "Containers/UnrealString.h"
+#include "Templates/Function.h"
+#include "UObject/Stack.h"
+
+#include "ToolsetRegistry/DelegateHandle.h"
+
+#define UE_API TOOLSETREGISTRY_API
+
+namespace UE::ToolsetRegistry
+{
+	/// Register for tracking of tool call exceptions for the lifetime of this object.
+	class FToolCallExceptionHandler
+	{
+	public:
+		UE_API FToolCallExceptionHandler();
+		~FToolCallExceptionHandler() = default;
+
+		// Disallow copy and move.
+		FToolCallExceptionHandler(const FToolCallExceptionHandler&) = delete;
+		FToolCallExceptionHandler& operator=(const FToolCallExceptionHandler&) = delete;
+		FToolCallExceptionHandler(FToolCallExceptionHandler&&) = delete;
+		FToolCallExceptionHandler& operator=(FToolCallExceptionHandler&&) = delete;
+
+		/// Description of script exceptions raised by UKismetSystemLibrary::RaiseScriptError or
+		/// FBlueprintCoreDelegates::ThrowScriptException via
+		/// FBlueprintCoreDelegates::OnScriptException during the lifetime of this object.
+		FString GetException() const
+		{
+			return ScriptException;
+		}
+
+		/// Call a function within a scope where it's possible to call
+		/// UKismetSystemLibrary::RaiseScriptError() without dynamically calling a UFunction.
+		/// 
+		/// When tools are executed via FBlueprintLibraryToolset, a stack frame is configured and
+		/// the appropriate method of a class derived from UToolsetDefinition is called.
+		/// In test cases it is convenient to be able to call a toolset method directly from C++
+		/// but this does not result in the configuration of a Blueprints stack frame (FFrame)
+		/// which results in UKismetSystemLibrary::RaiseScriptError() silently doing nothing.
+		/// 
+		/// This method calls the specified function within a scope that has a stack frame
+		/// constructed on the calling thread such that calls to
+		/// UKismetSystemLibrary::RaiseScriptError() within Func will notify this class of the
+		/// errors.
+		/// 
+		/// @warning This is intended to be only used for testing tools derived from
+		/// UToolsetDefinition.
+		/// 
+		/// @param Func Function to call within a test scope.
+		///
+		UE_API void CaptureErrorsIn(TFunction<void()>&& Func);
+
+	private:
+		// Track script exceptions raised by UKismetSystemLibrary::RaiseScriptError or
+		// FBlueprintCoreDelegates::ThrowScriptException via FBlueprintCoreDelegates::OnScriptException
+		// and copy into local storage.
+		void OnScriptException(
+			const UObject* ActiveObject, const FFrame& StackFrame, const FBlueprintExceptionInfo& Info);
+
+		// Build an exception message for the stack frame with the specified ExceptionTypeName.
+		FString GetBackupExceptionMessage(const FString& ExceptionTypeName, const FFrame& StackFrame);
+
+		// Try to build an exception message from exception info and stack frame.
+		TOptional<FString> GetBlueprintExceptionMessage(
+			const FBlueprintExceptionInfo& Info, const FFrame& StackFrame);
+
+		// Append an message to the current exception message.
+		void AddToExceptionMessage(const FString& NewMessage);
+
+		// Handler for FBlueprintCoreDelegates::OnScriptException
+		FDelegateHandleRaii ScriptExceptionHandle;
+		// Handler for FScriptExceptionHandler
+		FScopedScriptExceptionHandler RuntimeScriptExceptionHandle;
+		// Raised script exception message.
+		FString ScriptException;
+	};
+}
+
+#undef UE_API
