[Toolset Registry] Fix minor issues in exception handler test.

#Synced-CL 52353462

diff --git a/Engine/Plugins/Experimental/ToolsetRegistry/Source/ToolsetRegistry/Private/Tests/ToolCallExceptionHandlerTest.cpp b/Engine/Plugins/Experimental/ToolsetRegistry/Source/ToolsetRegistry/Private/Tests/ToolCallExceptionHandlerTest.cpp
--- a/Engine/Plugins/Experimental/ToolsetRegistry/Source/ToolsetRegistry/Private/Tests/ToolCallExceptionHandlerTest.cpp
+++ b/Engine/Plugins/Experimental/ToolsetRegistry/Source/ToolsetRegistry/Private/Tests/ToolCallExceptionHandlerTest.cpp
@@ -57,7 +57,7 @@ void FToolCallExceptionHandlerSpec::Define()
 			FString Error2 = TEXT("another error");
 			FToolCallExceptionHandler Handler;
 			Handler.CaptureErrorsIn(
-				[&Error1, Error2]()
+				[&Error1, &Error2]()
 				{
 					UKismetSystemLibrary::RaiseScriptError(Error1);
 					UKismetSystemLibrary::RaiseScriptError(Error2);
@@ -171,7 +171,7 @@ void FToolCallExceptionHandlerSpec::Define()
 			FJsonValueOrError Result = TestHelpers::ExecuteToolCallWithJson(
 				TestObject, TEXT("TestSpecObjectParam"), InvalidJson);
 
-			if (!TestTrue("Should have error", Result.HasError())) return;
+			if (!TestTrue(TEXT("Should have error"), Result.HasError())) return;
 
 			FString ErrorMessage = Result.GetError();
 
