[AI Toolsets]: Gracefully handle null object arguments in UToolsetLibrary

#Synced-CL 52369294

diff --git a/Engine/Plugins/Experimental/ToolsetRegistry/Source/ToolsetRegistry/Private/Tests/ToolsetLibraryTest.cpp b/Engine/Plugins/Experimental/ToolsetRegistry/Source/ToolsetRegistry/Private/Tests/ToolsetLibraryTest.cpp
--- a/Engine/Plugins/Experimental/ToolsetRegistry/Source/ToolsetRegistry/Private/Tests/ToolsetLibraryTest.cpp
+++ b/Engine/Plugins/Experimental/ToolsetRegistry/Source/ToolsetRegistry/Private/Tests/ToolsetLibraryTest.cpp
@@ -9,6 +9,7 @@
 #include "ToolsetJsonTest.h"
 #include "ToolsetRegistry/JsonConversion.h"
 #include "ToolsetRegistry/ToolsetLibrary.h"
+#include "ToolsetRegistry/ToolCallExceptionHandler.h"
 
 #if WITH_DEV_AUTOMATION_TESTS
 
@@ -19,6 +20,7 @@ END_DEFINE_SPEC(FToolsetLibrarySpec)
 
 void FToolsetLibrarySpec::Define()
 {
+	using namespace UE::ToolsetRegistry;
 	using namespace UE::ToolsetRegistry::Internal;
 
 	It("Generates schema from properties.", [this]()
@@ -63,6 +65,19 @@ void FToolsetLibrarySpec::Define()
 		TestTrue("SetProp", TestObject->VisibleProperty);
 	});
 
+	It(TEXT("GetObjectProperties raises error with null object."), [this]()
+	{
+		FToolCallExceptionHandler ExceptionHandler;
+		ExceptionHandler.CaptureErrorsIn(
+			[this]()
+			{
+				FString PropertiesJson = UToolsetLibrary::GetObjectProperties(
+					nullptr, TArray<FName>());
+				TestTrue(TEXT("No properties"), PropertiesJson.IsEmpty());
+			});
+		TestNotEqual(TEXT("Raised error"), *ExceptionHandler.GetException(), TEXT(""));
+	});
+
 	It("GetObjectProperties omits missing properties from result.", [this]()
 	{
 		UToolsetJsonTestObject* TestObject = NewObject<UToolsetJsonTestObject>();
@@ -138,12 +153,37 @@ void FToolsetLibrarySpec::Define()
 		TestTrue("Subclass", Subclasses.Contains(UFakeToolset::StaticClass()));
 	});
 
+	It(TEXT("GetDerivedClasses raises error with null object."), [this]()
+	{
+		FToolCallExceptionHandler ExceptionHandler;
+		ExceptionHandler.CaptureErrorsIn(
+			[this]()
+			{
+				TArray<UClass*> DerivedClasses = UToolsetLibrary::GetDerivedClasses(nullptr);
+				TestTrue(TEXT("No classes"), DerivedClasses.IsEmpty());
+			});
+		TestNotEqual(TEXT("Raised error"), *ExceptionHandler.GetException(), TEXT(""));
+	});
+
 	It("Gets substructs of a struct.", [this]()
 	{
 		TArray<UScriptStruct*> Substructs = UToolsetLibrary::GetDerivedStructs(
 			FAnimNode_Base::StaticStruct());
 		TestTrue("Subclass", Substructs.Contains(FAnimNode_Root::StaticStruct()));
 	});
+
+	It(TEXT("GetDerivedStructs raises error with null object."), [this]()
+	{
+		FToolCallExceptionHandler ExceptionHandler;
+		ExceptionHandler.CaptureErrorsIn(
+			[this]()
+			{
+				TArray<UScriptStruct*> DerivedStructs =
+					UToolsetLibrary::GetDerivedStructs(nullptr);
+				TestTrue(TEXT("No structs"), DerivedStructs.IsEmpty());
+			});
+		TestNotEqual(TEXT("Raised error"), *ExceptionHandler.GetException(), TEXT(""));
+	});
 }
 
 #endif
diff --git a/Engine/Plugins/Experimental/ToolsetRegistry/Source/ToolsetRegistry/Private/ToolsetRegistry/ToolsetLibrary.cpp b/Engine/Plugins/Experimental/ToolsetRegistry/Source/ToolsetRegistry/Private/ToolsetRegistry/ToolsetLibrary.cpp
--- a/Engine/Plugins/Experimental/ToolsetRegistry/Source/ToolsetRegistry/Private/ToolsetRegistry/ToolsetLibrary.cpp
+++ b/Engine/Plugins/Experimental/ToolsetRegistry/Source/ToolsetRegistry/Private/ToolsetRegistry/ToolsetLibrary.cpp
@@ -38,7 +38,11 @@ FString UToolsetLibrary::ListStructProperties(UStruct* Struct, bool bUserVisible
 
 FString UToolsetLibrary::GetObjectProperties(const UObject* Object, const TArray<FName>& PropertyNames)
 {
-	check(Object);
+	if (!Object)
+	{
+		UKismetSystemLibrary::RaiseScriptError(TEXT("Null object provided"));
+		return FString();
+	}
 	TSharedPtr<FJsonObject> PropertyData = MakeShared<FJsonObject>();
 	TArray<FString> InaccessibleProperties;
 	for (const auto& PropertyName : PropertyNames)
@@ -124,7 +128,12 @@ bool UToolsetLibrary::SetObjectProperties(UObject* Object, const FString& Proper
 
 TArray<UClass*> UToolsetLibrary::GetDerivedClasses(UClass* BaseClass)
 {
-	check(BaseClass);
+	TArray<UClass*> DerivedClasses;
+	if (!BaseClass)
+	{
+		UKismetSystemLibrary::RaiseScriptError(TEXT("Null base class provided"));
+		return DerivedClasses;
+	}
 
 	FAssetRegistryModule& AssetRegistryModule =
 		FModuleManager::LoadModuleChecked<FAssetRegistryModule>("AssetRegistry");
@@ -133,7 +142,6 @@ TArray<UClass*> UToolsetLibrary::GetDerivedClasses(UClass* BaseClass)
 	TSet<FTopLevelAssetPath> DerivedNames;
 	AssetRegistry.GetDerivedClassNames({ BaseClass->GetClassPathName() }, {}, DerivedNames);
 
-	TArray<UClass*> DerivedClasses;
 	for (const FTopLevelAssetPath& DerivedName : DerivedNames)
 	{
 		// We ignore skeleton classes (contained in BPs) but consider all other subclasses.
@@ -157,9 +165,13 @@ TArray<UClass*> UToolsetLibrary::GetDerivedClasses(UClass* BaseClass)
 
 TArray<UScriptStruct*> UToolsetLibrary::GetDerivedStructs(UScriptStruct* BaseStruct)
 {
-	check(BaseStruct);
-
 	TArray<UScriptStruct*> DerivedStructs;
+	if (!BaseStruct)
+	{
+		UKismetSystemLibrary::RaiseScriptError(TEXT("Null base struct provided"));
+		return DerivedStructs;
+	}
+	
 	for (TObjectIterator<UScriptStruct> It; It; ++It)
 	{
 		UScriptStruct* Struct = *It;
