Guard UToolsetRegistry::UnregisterToolsetClass when not UObjectInitialized

#Synced-CL 52419586

diff --git a/Engine/Plugins/Experimental/ToolsetRegistry/Source/ToolsetRegistry/Private/ToolsetRegistry/UToolsetRegistry.cpp b/Engine/Plugins/Experimental/ToolsetRegistry/Source/ToolsetRegistry/Private/ToolsetRegistry/UToolsetRegistry.cpp
--- a/Engine/Plugins/Experimental/ToolsetRegistry/Source/ToolsetRegistry/Private/ToolsetRegistry/UToolsetRegistry.cpp
+++ b/Engine/Plugins/Experimental/ToolsetRegistry/Source/ToolsetRegistry/Private/ToolsetRegistry/UToolsetRegistry.cpp
@@ -56,6 +56,11 @@ void UToolsetRegistry::RegisterToolsetClass(TSubclassOf<UToolsetDefinition> InTo
 
 void UToolsetRegistry::UnregisterToolsetClass(TSubclassOf<UToolsetDefinition> InToolsetClass)
 {
+	if (!UObjectInitialized())
+	{
+		return;
+	}
+
 	if (!InToolsetClass)
 	{
 		UE_LOG(LogToolsetRegistry, Warning, TEXT("Attempting to unregister null toolset class"));
