C++ bridge and data types for AIAssistant sandbox integration

Introduces the foundational C++ bridge functions that the AIAssistantWebApplication will use to interact with the FileSandbox module

#Synced-CL 52233446

diff --git a/Engine/Plugins/Experimental/ToolsetRegistry/Source/ToolsetRegistry/Private/Tests/SandboxLibraryTest.cpp b/Engine/Plugins/Experimental/ToolsetRegistry/Source/ToolsetRegistry/Private/Tests/SandboxLibraryTest.cpp
new file mode 100644
--- /dev/null
+++ b/Engine/Plugins/Experimental/ToolsetRegistry/Source/ToolsetRegistry/Private/Tests/SandboxLibraryTest.cpp
@@ -0,0 +1,100 @@
+// Copyright Epic Games, Inc. All Rights Reserved.
+
+#include "ToolsetRegistry/SandboxLibrary.h"
+#include "Tests/ToolsetRegistryTestFlags.h"
+
+#if WITH_DEV_AUTOMATION_TESTS
+
+using namespace UE::ToolsetRegistry;
+
+// -----------------------------------------------------------------------------
+// SandboxLibrary tests
+// -----------------------------------------------------------------------------
+
+BEGIN_DEFINE_SPEC(
+	FSandboxLibrarySpec,
+	"AI.ToolsetRegistry.Sandbox.Library",
+	ToolsetRegistryTest::Flags)
+END_DEFINE_SPEC(FSandboxLibrarySpec)
+
+void FSandboxLibrarySpec::Define()
+{
+	// ToolsetRegistry depends on FileSandboxCore, so the module is always available.
+
+	Describe(TEXT("When no sandbox is active"), [this]
+	{
+		It(TEXT("IsActive should return false"), [this]
+		{
+			if (FGlobalSandbox::IsActive()) return; // skip if a sandbox is unexpectedly active
+			TestFalse(TEXT("No active sandbox"), FGlobalSandbox::IsActive());
+		});
+
+		It(TEXT("GetActiveName should return empty"), [this]
+		{
+			if (FGlobalSandbox::IsActive()) return;
+			TestTrue(TEXT("Name is empty"), FGlobalSandbox::GetActiveName().IsEmpty());
+		});
+
+		It(TEXT("GetChanges should return empty array"), [this]
+		{
+			if (FGlobalSandbox::IsActive()) return;
+			TestEqual(TEXT("No changes"), FGlobalSandbox::GetChanges().Num(), 0);
+		});
+
+		It(TEXT("Leave should return true (no-op)"), [this]
+		{
+			if (FGlobalSandbox::IsActive()) return;
+			TestTrue(TEXT("Leave succeeds as no-op"), FGlobalSandbox::Leave());
+		});
+	});
+
+	Describe(TEXT("Sandbox lifecycle"), [this]
+	{
+		const FString TestSandboxName = TEXT("ToolsetRegistryTest_SandboxLibrary");
+
+		BeforeEach([this, TestSandboxName]
+		{
+			// Ensure no sandbox is active before each test.
+			FGlobalSandbox::Leave();
+		});
+
+		AfterEach([this, TestSandboxName]
+		{
+			// Always clean up: discard any sandbox changes, then leave.
+			if (FGlobalSandbox::IsActive())
+			{
+				FGlobalSandbox::Discard();
+			}
+			FGlobalSandbox::Leave();
+		});
+
+		It(TEXT("Enter creates a new sandbox and IsActive returns true"), [this, TestSandboxName]
+		{
+			TestTrue(TEXT("Enter succeeded"), FGlobalSandbox::Enter(TestSandboxName, TEXT("")));
+			TestTrue(TEXT("Sandbox is now active"), FGlobalSandbox::IsActive());
+			TestEqual(TEXT("Active sandbox has expected name"), FGlobalSandbox::GetActiveName(), TestSandboxName);
+		});
+
+		It(TEXT("Enter with the same name is a no-op when already active"), [this, TestSandboxName]
+		{
+			TestTrue(TEXT("First Enter succeeded"), FGlobalSandbox::Enter(TestSandboxName, TEXT("")));
+			TestTrue(TEXT("Second Enter succeeded"), FGlobalSandbox::Enter(TestSandboxName, TEXT("")));
+			TestTrue(TEXT("Sandbox still active"), FGlobalSandbox::IsActive());
+		});
+
+		It(TEXT("Leave deactivates the sandbox"), [this, TestSandboxName]
+		{
+			TestTrue(TEXT("Enter succeeded"), FGlobalSandbox::Enter(TestSandboxName, TEXT("")));
+			TestTrue(TEXT("Leave succeeded"), FGlobalSandbox::Leave());
+			TestFalse(TEXT("Sandbox no longer active"), FGlobalSandbox::IsActive());
+		});
+
+		It(TEXT("GetChanges returns empty array for a fresh sandbox"), [this, TestSandboxName]
+		{
+			TestTrue(TEXT("Enter succeeded"), FGlobalSandbox::Enter(TestSandboxName, TEXT("")));
+			TestEqual(TEXT("No changes in fresh sandbox"), FGlobalSandbox::GetChanges().Num(), 0);
+		});
+	});
+}
+
+#endif // WITH_DEV_AUTOMATION_TESTS
diff --git a/Engine/Plugins/Experimental/ToolsetRegistry/Source/ToolsetRegistry/Private/ToolsetRegistry/SandboxLibrary.cpp b/Engine/Plugins/Experimental/ToolsetRegistry/Source/ToolsetRegistry/Private/ToolsetRegistry/SandboxLibrary.cpp
new file mode 100644
--- /dev/null
+++ b/Engine/Plugins/Experimental/ToolsetRegistry/Source/ToolsetRegistry/Private/ToolsetRegistry/SandboxLibrary.cpp
@@ -0,0 +1,205 @@
+// Copyright Epic Games, Inc. All Rights Reserved.
+
+#include "ToolsetRegistry/SandboxLibrary.h"
+
+#include "ToolsetRegistry/Module.h"
+#include "FileSandboxCoreMinimal.h"
+#include "Modules/ModuleManager.h"
+#include "Data/SandboxMetaData.h"
+#include "Types/EBreakBehavior.h"
+#include "Types/SandboxedFileChangeInfo.h"
+#include "Types/Sandbox/PersistArgs.h"
+#include "Types/Sandbox/PersistResult.h"
+#include "Types/Sandbox/RevertArgs.h"
+#include "Types/Sandbox/RevertResult.h"
+
+namespace
+{
+	struct FSandboxAccess
+	{
+		UE::FileSandboxCore::ISandboxManager* Manager = nullptr;
+		UE::FileSandboxCore::ISandboxInstance* ActiveInstance = nullptr;
+	};
+
+	FSandboxAccess GetSandboxAccess(bool bLogErrorIfNoManager = false)
+	{
+		using namespace UE::FileSandboxCore;
+
+		if (!IFileSandboxCoreModule::IsAvailable())
+		{
+			if (bLogErrorIfNoManager)
+			{
+				UE_LOGF(LogToolsetRegistry, Error, "FileSandbox module is not available");
+			}
+			return {};
+		}
+
+		ISandboxManager* Manager = &IFileSandboxCoreModule::Get().GetSandboxManager();
+		return { Manager, Manager->GetActiveSandboxInstance() };
+	}
+}
+
+namespace UE::ToolsetRegistry
+{
+	bool FGlobalSandbox::IsActive()
+	{
+		return GetSandboxAccess().ActiveInstance != nullptr;
+	}
+
+	FString FGlobalSandbox::GetActiveName()
+	{
+		FSandboxAccess Access = GetSandboxAccess();
+		return Access.ActiveInstance ? Access.ActiveInstance->GetInitialMetaData().Name : FString();
+	}
+
+	bool FGlobalSandbox::Enter(const FString& Name, const FString& Description)
+	{
+		using namespace UE::FileSandboxCore;
+
+		FSandboxAccess Access = GetSandboxAccess(/*bLogErrorIfNoManager=*/true);
+		if (!Access.Manager)
+		{
+			return false;
+		}
+
+		ISandboxManager* Manager = Access.Manager;
+
+		// If already in the requested sandbox, nothing to do.
+		if (Access.ActiveInstance)
+		{
+			if (Access.ActiveInstance->GetInitialMetaData().Name == Name)
+			{
+				return true;
+			}
+
+			// A different sandbox is active — leave it before entering the new one.
+			FLeaveSandboxResult LeaveResult = Manager->LeaveSandbox();
+			if (LeaveResult.ErrorCode != ELeaveSandboxErrorCode::Success)
+			{
+				UE_LOGF(LogToolsetRegistry, Error, "Failed to leave active sandbox before entering '%ls'", *Name);
+				return false;
+			}
+		}
+
+		// Try to load an existing sandbox by name.
+		FLoadSandboxResult LoadResult = Manager->LoadNamedSandbox(FLoadSandboxByNameArgs(Name));
+		if (LoadResult.HasValue())
+		{
+			UE_LOGF(LogToolsetRegistry, Display, "Resumed existing sandbox: %ls", *Name);
+			return true;
+		}
+
+		// Create a new sandbox.
+		FNewSandboxArgs CreateArgs(Name, Description);
+		FNewSandboxResult CreateResult = Manager->CreateNewSandbox(CreateArgs);
+		if (CreateResult.HasValue())
+		{
+			UE_LOGF(LogToolsetRegistry, Display, "Created new sandbox: %ls", *Name);
+			return true;
+		}
+
+		UE_LOGF(LogToolsetRegistry, Error, "Failed to enter sandbox: %ls", *Name);
+		return false;
+	}
+
+	bool FGlobalSandbox::Leave()
+	{
+		using namespace UE::FileSandboxCore;
+
+		FSandboxAccess Access = GetSandboxAccess();
+		if (!Access.Manager || !Access.ActiveInstance)
+		{
+			return true;
+		}
+
+		FLeaveSandboxResult Result = Access.Manager->LeaveSandbox();
+		return Result.ErrorCode == ELeaveSandboxErrorCode::Success;
+	}
+
+	TArray<UE::FileSandboxCore::FSandboxedFileChangeInfo> FGlobalSandbox::GetChanges()
+	{
+		using namespace UE::FileSandboxCore;
+
+		TArray<FSandboxedFileChangeInfo> Changes;
+		ISandboxInstance* Instance = GetSandboxAccess().ActiveInstance;
+		if (!Instance)
+		{
+			return {};
+		}
+
+		Instance->EnumerateFileChanges(
+			[&Changes](const FSandboxedFileChangeInfo& Info) -> EBreakBehavior
+			{
+				Changes.Add(Info);
+				return EBreakBehavior::Continue;
+			},
+			EFileEnumerationFlags::All);
+
+		return Changes;
+	}
+
+	bool FGlobalSandbox::Persist(const TArray<FString>& Files)
+	{
+		using namespace UE::FileSandboxCore;
+
+		ISandboxInstance* Instance = GetSandboxAccess().ActiveInstance;
+		if (!Instance)
+		{
+			UE_LOGF(LogToolsetRegistry, Error, "No active sandbox to persist");
+			return false;
+		}
+
+		if (Files.IsEmpty())
+		{
+			return Instance->PersistAll();
+		}
+
+		FPersistArgs Args;
+		Args.Files = Files;
+		FPersistResult Result = Instance->PersistSandbox(Args);
+		return Result.PersistStatus == EPersistStatus::Success;
+	}
+
+	bool FGlobalSandbox::Discard()
+	{
+		using namespace UE::FileSandboxCore;
+
+		ISandboxInstance* Instance = GetSandboxAccess().ActiveInstance;
+		if (!Instance)
+		{
+			UE_LOGF(LogToolsetRegistry, Error, "No active sandbox to discard");
+			return false;
+		}
+
+		Instance->RevertAll();
+		return true;
+	}
+
+	bool FGlobalSandbox::DiscardFiles(const TArray<FString>& Files)
+	{
+		if (Files.IsEmpty())
+		{
+			UE_LOGF(LogToolsetRegistry, Error,
+				"DiscardFiles called with empty file list; ignoring to avoid discarding all sandbox changes");
+			return false;
+		}
+
+		using namespace UE::FileSandboxCore;
+
+		ISandboxInstance* Instance = GetSandboxAccess().ActiveInstance;
+		if (!Instance)
+		{
+			UE_LOGF(LogToolsetRegistry, Error, "No active sandbox to discard files from");
+			return false;
+		}
+
+		FRevertResult Result = Instance->RevertSandbox(FRevertArgs(Files));
+		// FRevertResult carries only packages pending hot-reload or purge - it has no
+		// error field, so the operation is considered to always succeed once the instance
+		// is available. Log the counts for diagnostics.
+		UE_LOGF(LogToolsetRegistry, Verbose,
+			"RevertSandbox: %d package(s) pending hot-reload, %d package(s) pending purge",
+			Result.PackagesPendingHotReload.Num(), Result.PackagesPendingPurge.Num());
+		return true;
+	}
+}
diff --git a/Engine/Plugins/Experimental/ToolsetRegistry/Source/ToolsetRegistry/Public/ToolsetRegistry/SandboxLibrary.h b/Engine/Plugins/Experimental/ToolsetRegistry/Source/ToolsetRegistry/Public/ToolsetRegistry/SandboxLibrary.h
new file mode 100644
--- /dev/null
+++ b/Engine/Plugins/Experimental/ToolsetRegistry/Source/ToolsetRegistry/Public/ToolsetRegistry/SandboxLibrary.h
@@ -0,0 +1,73 @@
+// Copyright Epic Games, Inc. All Rights Reserved.
+
+#pragma once
+
+#include "CoreMinimal.h"
+#include "Types/SandboxedFileChangeInfo.h"
+
+namespace UE::ToolsetRegistry
+{
+	/** Provides static access to the globally active FileSandbox instance. */
+	class TOOLSETREGISTRY_API FGlobalSandbox
+	{
+	public:
+		/** Returns whether a sandbox is currently active. */
+		static bool IsActive();
+
+		/** Returns the name of the active sandbox, or an empty string if none is active. */
+		static FString GetActiveName();
+
+		/**
+		 * Creates a new sandbox or resumes an existing one by name.
+		 * If a different sandbox is already active, it is left before entering the new one.
+		 *
+		 * @param Name        The name of the sandbox to enter or create.
+		 * @param Description Used only when creating a new sandbox; ignored if a sandbox
+		 *                    with the given name already exists.
+		 * @return True if the named sandbox is now active; false if leaving the current
+		 *         sandbox or entering the requested sandbox failed.
+		 */
+		static bool Enter(const FString& Name, const FString& Description);
+
+		/**
+		 * Leaves the active sandbox without deleting it. The sandbox's files are preserved.
+		 *
+		 * @return True if no sandbox is active after the call.
+		 */
+		static bool Leave();
+
+		/**
+		 * Returns the list of files changed in the active sandbox.
+		 * Returns an empty array if no sandbox is active.
+		 */
+		static TArray<UE::FileSandboxCore::FSandboxedFileChangeInfo> GetChanges();
+
+		/**
+		 * Persists sandbox changes to the real filesystem.
+		 *
+		 * @param Files Files to persist. If empty, all changes are persisted.
+		 * @return True on success.
+		 */
+		static bool Persist(const TArray<FString>& Files);
+
+		/**
+		 * Discards all changes in the active sandbox.
+		 * The sandbox remains active after the call; call Leave() to exit it.
+		 *
+		 * @return True on success.
+		 */
+		static bool Discard();
+
+		/**
+		 * Discards only the specified files in the active sandbox.
+		 * To discard all changes, call Discard() instead.
+		 *
+		 * @param Files The files to discard. Must be non-empty; passing an empty list
+		 *              is treated as a caller error and returns false to avoid silently
+		 *              discarding all sandbox changes.
+		 * @return True if the files were reverted; false if no sandbox is active or
+		 *         Files is empty.
+		 */
+		static bool DiscardFiles(const TArray<FString>& Files);
+	};
+}
diff --git a/Engine/Plugins/Experimental/ToolsetRegistry/Source/ToolsetRegistry/ToolsetRegistry.Build.cs b/Engine/Plugins/Experimental/ToolsetRegistry/Source/ToolsetRegistry/ToolsetRegistry.Build.cs
--- a/Engine/Plugins/Experimental/ToolsetRegistry/Source/ToolsetRegistry/ToolsetRegistry.Build.cs
+++ b/Engine/Plugins/Experimental/ToolsetRegistry/Source/ToolsetRegistry/ToolsetRegistry.Build.cs
@@ -59,6 +59,7 @@ public class ToolsetRegistry : ModuleRules
 				"RHI",
 				"StatusBar",
 				"UnrealEd",
+				"FileSandboxCore",
 			}
 			);
 
diff --git a/Engine/Plugins/Experimental/ToolsetRegistry/ToolsetRegistry.uplugin b/Engine/Plugins/Experimental/ToolsetRegistry/ToolsetRegistry.uplugin
--- a/Engine/Plugins/Experimental/ToolsetRegistry/ToolsetRegistry.uplugin
+++ b/Engine/Plugins/Experimental/ToolsetRegistry/ToolsetRegistry.uplugin
@@ -33,6 +33,10 @@
 		{
 			"Name": "EditorScriptingUtilities",
 			"Enabled": true
+		},
+		{
+			"Name": "FileSandbox",
+			"Enabled": true
 		}
 	]
 }
\ No newline at end of file
