Fix FGlobalSandbox export tags

#Synced-CL 52233446

diff --git a/Engine/Plugins/Experimental/ToolsetRegistry/Source/ToolsetRegistry/Public/ToolsetRegistry/SandboxLibrary.h b/Engine/Plugins/Experimental/ToolsetRegistry/Source/ToolsetRegistry/Public/ToolsetRegistry/SandboxLibrary.h
--- a/Engine/Plugins/Experimental/ToolsetRegistry/Source/ToolsetRegistry/Public/ToolsetRegistry/SandboxLibrary.h
+++ b/Engine/Plugins/Experimental/ToolsetRegistry/Source/ToolsetRegistry/Public/ToolsetRegistry/SandboxLibrary.h
@@ -5,17 +5,19 @@
 #include "CoreMinimal.h"
 #include "Types/SandboxedFileChangeInfo.h"
 
+#define UE_API TOOLSETREGISTRY_API
+
 namespace UE::ToolsetRegistry
 {
 	/** Provides static access to the globally active FileSandbox instance. */
-	class TOOLSETREGISTRY_API FGlobalSandbox
+	class FGlobalSandbox
 	{
 	public:
 		/** Returns whether a sandbox is currently active. */
-		static bool IsActive();
+		static UE_API bool IsActive();
 
 		/** Returns the name of the active sandbox, or an empty string if none is active. */
-		static FString GetActiveName();
+		static UE_API FString GetActiveName();
 
 		/**
 		 * Creates a new sandbox or resumes an existing one by name.
@@ -27,20 +29,20 @@ namespace UE::ToolsetRegistry
 		 * @return True if the named sandbox is now active; false if leaving the current
 		 *         sandbox or entering the requested sandbox failed.
 		 */
-		static bool Enter(const FString& Name, const FString& Description);
+		static UE_API bool Enter(const FString& Name, const FString& Description);
 
 		/**
 		 * Leaves the active sandbox without deleting it. The sandbox's files are preserved.
 		 *
 		 * @return True if no sandbox is active after the call.
 		 */
-		static bool Leave();
+		static UE_API bool Leave();
 
 		/**
 		 * Returns the list of files changed in the active sandbox.
 		 * Returns an empty array if no sandbox is active.
 		 */
-		static TArray<UE::FileSandboxCore::FSandboxedFileChangeInfo> GetChanges();
+		static UE_API TArray<UE::FileSandboxCore::FSandboxedFileChangeInfo> GetChanges();
 
 		/**
 		 * Persists sandbox changes to the real filesystem.
@@ -48,7 +50,7 @@ namespace UE::ToolsetRegistry
 		 * @param Files Files to persist. If empty, all changes are persisted.
 		 * @return True on success.
 		 */
-		static bool Persist(const TArray<FString>& Files);
+		static UE_API bool Persist(const TArray<FString>& Files);
 
 		/**
 		 * Discards all changes in the active sandbox.
@@ -56,7 +58,7 @@ namespace UE::ToolsetRegistry
 		 *
 		 * @return True on success.
 		 */
-		static bool Discard();
+		static UE_API bool Discard();
 
 		/**
 		 * Discards only the specified files in the active sandbox.
@@ -68,6 +70,8 @@ namespace UE::ToolsetRegistry
 		 * @return True if the files were reverted; false if no sandbox is active or
 		 *         Files is empty.
 		 */
-		static bool DiscardFiles(const TArray<FString>& Files);
+		static UE_API bool DiscardFiles(const TArray<FString>& Files);
 	};
 }
+
+#undef UE_API
