Add a couple of editor app tools that will help provide context about the state of the editor.
* CaptureEditorImage which generates an image that matches what how the user sees the entire editor app.
* GetOpenAssets which returns a list of assets that are open in asset editors.

#Synced-CL 52563198

diff --git a/Engine/Plugins/Experimental/ToolsetRegistry/Source/ToolsetRegistry/Private/Tests/EditorAppToolsetTest.cpp b/Engine/Plugins/Experimental/ToolsetRegistry/Source/ToolsetRegistry/Private/Tests/EditorAppToolsetTest.cpp
--- a/Engine/Plugins/Experimental/ToolsetRegistry/Source/ToolsetRegistry/Private/Tests/EditorAppToolsetTest.cpp
+++ b/Engine/Plugins/Experimental/ToolsetRegistry/Source/ToolsetRegistry/Private/Tests/EditorAppToolsetTest.cpp
@@ -6,6 +6,7 @@
 #include "AssetRegistry/IAssetRegistry.h"
 #include "Editor.h"
 #include "Engine/World.h"
+#include "Subsystems/AssetEditorSubsystem.h"
 #include "GameFramework/Actor.h"
 #include "Misc/App.h"
 #include "Misc/CoreMiscDefines.h"
@@ -162,82 +163,92 @@ void FEditorAppToolsetSpec::Define()
 		TestFalse("Error is set", Result->Error.IsEmpty());
 	});
 
-	It(TEXT("CaptureAssetImage renders a thumbnail for a static mesh"),
-		[this]()
+	It(TEXT("CaptureAssetImage renders a thumbnail for a static mesh"), [this]()
+	{
+		if (!FApp::CanEverRender())
 		{
-			if (!FApp::CanEverRender())
-			{
-				// Rendering is disabled; skipping thumbnail test.
-				return;
-			}
+			// Rendering is disabled; skipping thumbnail test.
+			return;
+		}
 
-			UToolCallAsyncResultImage* Result = UEditorAppToolset::CaptureAssetImage(
-				TEXT("/Engine/BasicShapes/Cube.Cube"));
-			if (!TestNotNull("Result", Result))
-			{
-				return;
-			}
+		UToolCallAsyncResultImage* Result = UEditorAppToolset::CaptureAssetImage(
+			TEXT("/Engine/BasicShapes/Cube.Cube"));
+		if (!TestNotNull("Result", Result))
+		{
+			return;
+		}
 
-			TStrongObjectPtr<UToolCallAsyncResultImage> StrongResult(Result);
+		TStrongObjectPtr<UToolCallAsyncResultImage> StrongResult(Result);
 
-			AddCommand(new FFunctionLatentCommand(
-				[this, StrongResult]() mutable -> bool
+		AddCommand(new FFunctionLatentCommand(
+			[this, StrongResult]() mutable -> bool
+			{
+				if (!StrongResult->bIsComplete)
 				{
-					if (!StrongResult->bIsComplete)
-					{
-						return false; // not done yet
-					}
-					TestTrue("No error", StrongResult->Error.IsEmpty());
-					TestFalse("Has image data", StrongResult->Value.Data.IsEmpty());
-					TestEqual("MIME type",
-						StrongResult->Value.MimeType, FString(TEXT("image/png")));
-					return true; // done
-				}));
-		});
-
-	It(TEXT("CaptureAssetImage screenshots the viewport for the current level"),
-		[this]()
+					return false; // not done yet
+				}
+				TestTrue("No error", StrongResult->Error.IsEmpty());
+				TestFalse("Has image data", StrongResult->Value.Data.IsEmpty());
+				TestEqual("MIME type",
+					StrongResult->Value.MimeType, FString(TEXT("image/png")));
+				return true; // done
+			}));
+	});
+
+	It(TEXT("CaptureAssetImage screenshots the viewport for the current level"), [this]()
+	{
+		// TODO: We still need to root cause why this hangs indefinitely on build machines.
+		// Viewport screenshots are performed in the editor viewport render method so it's
+		// possible that is being disabled - somehow - when running on the build machine.
+		if (GIsBuildMachine || !FApp::CanEverRender())
 		{
-			// TODO: We still need to root cause why this hangs indefinitely on build machines.
-			// Viewport screenshots are performed in the editor viewport render method so it's
-			// possible that is being disabled - somehow - when running on the build machine.
-			if (GIsBuildMachine || !FApp::CanEverRender())
-			{
-				// Rendering is disabled; skipping thumbnail test.
-				return;
-			}
+			// Rendering is disabled; skipping thumbnail test.
+			return;
+		}
 
-			UWorld* EditorWorld = GEditor ? GEditor->GetEditorWorldContext().World() : nullptr;
-			if (!EditorWorld)
-			{
-				// No editor world available; skipping viewport screenshot test.
-				return;
-			}
+		UWorld* EditorWorld = GEditor ? GEditor->GetEditorWorldContext().World() : nullptr;
+		if (!EditorWorld)
+		{
+			// No editor world available; skipping viewport screenshot test.
+			return;
+		}
 
-			FString CurrentLevelPath = EditorWorld->GetOutermost()->GetName();
-			UToolCallAsyncResultImage* Result =
-				UEditorAppToolset::CaptureAssetImage(CurrentLevelPath);
-			if (!TestNotNull("Result", Result))
-			{
-				return;
-			}
+		FString CurrentLevelPath = EditorWorld->GetOutermost()->GetName();
+		UToolCallAsyncResultImage* Result =
+			UEditorAppToolset::CaptureAssetImage(CurrentLevelPath);
+		if (!TestNotNull("Result", Result))
+		{
+			return;
+		}
 
-			TStrongObjectPtr<UToolCallAsyncResultImage> StrongResult(Result);
+		TStrongObjectPtr<UToolCallAsyncResultImage> StrongResult(Result);
 
-			AddCommand(new FFunctionLatentCommand(
-				[this, StrongResult]() mutable -> bool
+		AddCommand(new FFunctionLatentCommand(
+			[this, StrongResult]() mutable -> bool
+			{
+				if (!StrongResult->bIsComplete)
 				{
-					if (!StrongResult->bIsComplete)
-					{
-						return false; // not done yet
-					}
-					TestTrue("No error", StrongResult->Error.IsEmpty());
-					TestFalse("Has image data", StrongResult->Value.Data.IsEmpty());
-					TestEqual("MIME type",
-						StrongResult->Value.MimeType, FString(TEXT("image/png")));
-					return true; // done
-				}));
-		});
+					return false; // not done yet
+				}
+				TestTrue("No error", StrongResult->Error.IsEmpty());
+				TestFalse("Has image data", StrongResult->Value.Data.IsEmpty());
+				TestEqual("MIME type",
+					StrongResult->Value.MimeType, FString(TEXT("image/png")));
+				return true; // done
+			}));
+	});
+
+	It(TEXT("CaptureEditorImage screenshots the entire editor window"), [this]()
+	{
+		if (GIsBuildMachine || !FApp::CanEverRender())
+		{
+			return;
+		}
+
+		FToolsetImage Result = UEditorAppToolset::CaptureEditorImage();
+		TestFalse("Has image data", Result.Data.IsEmpty());
+		TestEqual("MIME type", Result.MimeType, FString(TEXT("image/png")));
+	});
 
 	// Actor selection tests
 
@@ -408,6 +419,30 @@ void FEditorAppToolsetSpec::Define()
 			}));
 	});
 
+	It(TEXT("GetOpenAssets returns assets open in asset editors"), [this]()
+	{
+		if (FApp::IsUnattended())
+		{
+			return;
+		}
+
+		// Open a known engine asset, verify it appears in GetOpenAssets, then close it.
+		const FString AssetPath = TEXT("/Engine/BasicShapes/Cube");
+		UObject* Asset = StaticLoadObject(UObject::StaticClass(), nullptr, TEXT("/Engine/BasicShapes/Cube.Cube"));
+		if (!TestNotNull("Cube asset", Asset)) return;
+
+		UAssetEditorSubsystem* Subsystem = GEditor
+			? GEditor->GetEditorSubsystem<UAssetEditorSubsystem>() : nullptr;
+		if (!TestNotNull("AssetEditorSubsystem", Subsystem)) return;
+
+		Subsystem->OpenEditorForAsset(Asset);
+
+		ON_SCOPE_EXIT { Subsystem->CloseAllEditorsForAsset(Asset); };
+
+		TestTrue("Cube is listed",
+			UEditorAppToolset::GetOpenAssets().Contains(AssetPath));
+	});
+
 	It(TEXT("OpenEditorForAsset raises a script error for a non-existent asset"), [this]()
 	{
 		if (FApp::IsUnattended())
diff --git a/Engine/Plugins/Experimental/ToolsetRegistry/Source/ToolsetRegistry/Private/ToolsetRegistry/EditorAppToolset.cpp b/Engine/Plugins/Experimental/ToolsetRegistry/Source/ToolsetRegistry/Private/ToolsetRegistry/EditorAppToolset.cpp
--- a/Engine/Plugins/Experimental/ToolsetRegistry/Source/ToolsetRegistry/Private/ToolsetRegistry/EditorAppToolset.cpp
+++ b/Engine/Plugins/Experimental/ToolsetRegistry/Source/ToolsetRegistry/Private/ToolsetRegistry/EditorAppToolset.cpp
@@ -7,6 +7,8 @@
 #include "AssetRegistry/IAssetRegistry.h"
 #include "ContentBrowserModule.h"
 #include "Containers/Ticker.h"
+#include "Framework/Application/SlateApplication.h"
+#include "Widgets/SWindow.h"
 #include "Editor.h"
 #include "EditorAssetLibrary.h"
 #include "IContentBrowserSingleton.h"
@@ -293,6 +295,8 @@ namespace
 		FTSTicker::FDelegateHandle TickerHandle;
 		double StartTime;
 	};
+
+
 }
 
 FString UEditorAppToolset::SearchCVars(const FString& Name)
@@ -714,3 +718,123 @@ void UEditorAppToolset::OpenEditorForAsset(const FString& AssetPath)
 			FString::Printf(TEXT("Failed to open editor for asset: %s"), *AssetPath));
 	}
 }
+
+TArray<FString> UEditorAppToolset::GetOpenAssets()
+{
+	TArray<FString> Result;
+	if (!GEditor)
+	{
+		UKismetSystemLibrary::RaiseScriptError(TEXT("Editor not found."));
+		return Result;
+	}
+	UAssetEditorSubsystem* Subsystem = GEditor->GetEditorSubsystem<UAssetEditorSubsystem>();
+	check(Subsystem);
+	for (UObject* Asset : Subsystem->GetAllEditedAssets())
+	{
+		if (Asset)
+		{
+			Result.Add(FPackageName::ObjectPathToPackageName(Asset->GetPathName()));
+		}
+	}
+	return Result;
+}
+
+FToolsetImage UEditorAppToolset::CaptureEditorImage()
+{
+	// Capture every visible Slate window (main window, asset editors, detached panels, etc.)
+	// back-to-front, then composite them into a single image using their screen positions.
+	struct FCapturedWindow
+	{
+		FSlateRect ScreenRect;
+		TArray<FColor> Colors;
+		FIntVector Size;
+	};
+
+	TArray<TSharedRef<SWindow>> Windows;
+	FSlateApplication::Get().GetAllVisibleWindowsOrdered(Windows);
+
+	TArray<FCapturedWindow> Captured;
+	float MinX = TNumericLimits<float>::Max();
+	float MinY = TNumericLimits<float>::Max();
+	float MaxX = TNumericLimits<float>::Lowest();
+	float MaxY = TNumericLimits<float>::Lowest();
+
+	for (const TSharedRef<SWindow>& Window : Windows)
+	{
+		FCapturedWindow& Entry = Captured.AddDefaulted_GetRef();
+		Entry.ScreenRect = Window->GetRectInScreen();
+		if (!FSlateApplication::Get().TakeScreenshot(Window, Entry.Colors, Entry.Size) ||
+			Entry.Colors.IsEmpty())
+		{
+			Captured.Pop();
+			continue;
+		}
+		MinX = FMath::Min(MinX, Entry.ScreenRect.Left);
+		MinY = FMath::Min(MinY, Entry.ScreenRect.Top);
+		MaxX = FMath::Max(MaxX, Entry.ScreenRect.Left + Entry.Size.X);
+		MaxY = FMath::Max(MaxY, Entry.ScreenRect.Top + Entry.Size.Y);
+	}
+
+	if (Captured.IsEmpty())
+	{
+		UKismetSystemLibrary::RaiseScriptError(TEXT("Failed to capture any editor windows."));
+		return FToolsetImage{};
+	}
+
+	const FIntPoint TotalSize(
+		FMath::Max(1, FMath::RoundToInt(MaxX - MinX)),
+		FMath::Max(1, FMath::RoundToInt(MaxY - MinY)));
+
+	// Blit each window into the canvas at its screen-relative offset, back-to-front.
+	TArray<FColor> Canvas;
+	Canvas.SetNumZeroed(TotalSize.X * TotalSize.Y);
+	for (const FCapturedWindow& Win : Captured)
+	{
+		const int32 OffsetX = FMath::RoundToInt(Win.ScreenRect.Left - MinX);
+		const int32 OffsetY = FMath::RoundToInt(Win.ScreenRect.Top - MinY);
+		for (int32 Row = 0; Row < Win.Size.Y; ++Row)
+		{
+			const int32 CanvasY = OffsetY + Row;
+			if (CanvasY < 0 || CanvasY >= TotalSize.Y) continue;
+			const int32 SrcX0 = FMath::Max(0, -OffsetX);
+			const int32 SrcX1 = FMath::Min(Win.Size.X, TotalSize.X - OffsetX);
+			if (SrcX0 >= SrcX1) continue;
+			FMemory::Memcpy(
+				&Canvas[CanvasY * TotalSize.X + OffsetX + SrcX0],
+				&Win.Colors[Row * Win.Size.X + SrcX0],
+				(SrcX1 - SrcX0) * sizeof(FColor));
+		}
+	}
+
+	FIntPoint Dimensions = TotalSize;
+	TArray<FColor> ScaledPixels;
+	const int32 MaxImageSize = 1280;
+	if (Dimensions.X > MaxImageSize || Dimensions.Y > MaxImageSize)
+	{
+		const float Scale = (Dimensions.X >= Dimensions.Y)
+			? (float)MaxImageSize / Dimensions.X
+			: (float)MaxImageSize / Dimensions.Y;
+		const FIntPoint Scaled(
+			FMath::Max(1, (int32)(Dimensions.X * Scale)),
+			FMath::Max(1, (int32)(Dimensions.Y * Scale)));
+		ScaledPixels.SetNumUninitialized(Scaled.X * Scaled.Y);
+		for (int32 Y = 0; Y < Scaled.Y; ++Y)
+		{
+			for (int32 X = 0; X < Scaled.X; ++X)
+			{
+				ScaledPixels[Y * Scaled.X + X] =
+					Canvas[(int32)(Y / Scale) * Dimensions.X + (int32)(X / Scale)];
+			}
+		}
+		Dimensions = Scaled;
+	}
+
+	FToolsetImage Image;
+	const TArray<FColor>& Final = ScaledPixels.Num() > 0 ? ScaledPixels : Canvas;
+	if (!Image.SetFromBitmap(Final, Dimensions))
+	{
+		UKismetSystemLibrary::RaiseScriptError(TEXT("Failed to encode editor screenshot."));
+		return FToolsetImage{};
+	}
+	return Image;
+}
diff --git a/Engine/Plugins/Experimental/ToolsetRegistry/Source/ToolsetRegistry/Private/ToolsetRegistry/EditorAppToolset.h b/Engine/Plugins/Experimental/ToolsetRegistry/Source/ToolsetRegistry/Private/ToolsetRegistry/EditorAppToolset.h
--- a/Engine/Plugins/Experimental/ToolsetRegistry/Source/ToolsetRegistry/Private/ToolsetRegistry/EditorAppToolset.h
+++ b/Engine/Plugins/Experimental/ToolsetRegistry/Source/ToolsetRegistry/Private/ToolsetRegistry/EditorAppToolset.h
@@ -4,6 +4,7 @@
 
 #include "CoreMinimal.h"
 
+#include "ToolsetRegistry/ToolsetImage.h"
 #include "ToolsetRegistry/UToolsetRegistry.h"
 
 #include "EditorAppToolset.generated.h"
@@ -41,6 +42,13 @@ public:
 	UFUNCTION(BlueprintCallable, Category = "EditorAppToolset")
 	static UToolCallAsyncResultImage* CaptureAssetImage(const FString& AssetPath);
 
+	/*
+	 * Captures an image of the entire editor application as the user sees it.
+	 * @return An image of editor windows as they appear on the users' desktop.
+	 */
+	UFUNCTION(BlueprintCallable, Category = "EditorAppToolset")
+	static FToolsetImage CaptureEditorImage();
+
 	/*
 	 * Gets the currently selected actors in the level editor.
 	 * @return A list of the currently selected actors.
@@ -138,6 +146,13 @@ public:
 	UFUNCTION(BlueprintCallable, Category = "EditorAppToolset")
 	static void OpenEditorForAsset(const FString& AssetPath);
 
+	/*
+	 * Gets the list of assets currently open in asset editors.
+	 * @return A list of package paths to the assets open in asset editors.
+	 */
+	UFUNCTION(BlueprintCallable, Category = "EditorAppToolset")
+	static TArray<FString> GetOpenAssets();
+
 private:
 	static TSharedPtr<FJsonObject> CVarToJson(IConsoleObject* CVar);
 
diff --git a/Engine/Plugins/Experimental/ToolsetRegistry/Source/ToolsetRegistry/ToolsetRegistry.Build.cs b/Engine/Plugins/Experimental/ToolsetRegistry/Source/ToolsetRegistry/ToolsetRegistry.Build.cs
--- a/Engine/Plugins/Experimental/ToolsetRegistry/Source/ToolsetRegistry/ToolsetRegistry.Build.cs
+++ b/Engine/Plugins/Experimental/ToolsetRegistry/Source/ToolsetRegistry/ToolsetRegistry.Build.cs
@@ -45,6 +45,8 @@ public class ToolsetRegistry : ModuleRules
 				"CoreUObject",
 				"DeveloperSettings",
 				"EditorFramework",
+				"Slate",
+				"SlateCore",
 				"EditorScriptingUtilities",
 				"EditorSubsystem",
 				"Engine",
