[AI Assistant]: Fix LogsToolsetTest failure.

* Log limit test would fail if not enough log entries exist
* Utilized `FToolCallExceptionHandler::CaptureErrorsIn()` to avoid marshalling tool calls via JSON
* Fixed all TestTrue / TestFalse assertions so they use a more specific assertion or add context upon failure.
* Wrapped all strings with TEXT where required.

#Synced-CL 52565678

diff --git a/Engine/Plugins/Experimental/ToolsetRegistry/Source/ToolsetRegistry/Private/Tests/LogsToolsetTest.cpp b/Engine/Plugins/Experimental/ToolsetRegistry/Source/ToolsetRegistry/Private/Tests/LogsToolsetTest.cpp
--- a/Engine/Plugins/Experimental/ToolsetRegistry/Source/ToolsetRegistry/Private/Tests/LogsToolsetTest.cpp
+++ b/Engine/Plugins/Experimental/ToolsetRegistry/Source/ToolsetRegistry/Private/Tests/LogsToolsetTest.cpp
@@ -1,40 +1,39 @@
 // Copyright Epic Games, Inc. All Rights Reserved.
 
 #include "CoreMinimal.h"
+#include "Logging/LogMacros.h"
 #include "Misc/AutomationTest.h"
+#include "Misc/CString.h"
 #include "Misc/Guid.h"
 
 #include "Tests/ToolsetRegistryTestFlags.h"
 #include "ToolsetRegistry/LogsToolset.h"
-#include "ToolsetRegistry/ObjectFunctionToolCall.h"
+#include "ToolsetRegistry/Module.h"
 #include "ToolsetRegistry/ToolCallExceptionHandler.h"
 
 #if WITH_DEV_AUTOMATION_TESTS
 
 DEFINE_LOG_CATEGORY_STATIC(LogLogsToolsetTest, Log, All);
 
-BEGIN_DEFINE_SPEC(FLogsToolsetSpec, "AI.ToolsetRegistry.LogsToolsetSpec",
+BEGIN_DEFINE_SPEC(FLogsToolsetSpec, "AI.ToolsetRegistry.LogToolset",
 	ToolsetRegistryTest::Flags)
-	FString CallAndCaptureException(const FString& FunctionName, const FString& ParamsJson);
+
+	bool TestContains(
+		const FString& Context, const FString& ToSearch, const FString& ToFind,
+		ESearchCase::Type Case = ESearchCase::CaseSensitive)
+	{
+		return TestTrue(
+			FString::Printf(TEXT("%s: `%s` in `%s`"), *Context, *ToFind, *ToSearch),
+			ToSearch.Contains(ToFind, Case));
+	}
 END_DEFINE_SPEC(FLogsToolsetSpec)
 
-FString FLogsToolsetSpec::CallAndCaptureException(
-	const FString& FunctionName, const FString& ParamsJson)
-{
-	ULogsToolset* Toolset = NewObject<ULogsToolset>();
-	UFunction* Function = Toolset->GetClass()->FindFunctionByName(*FunctionName);
-	if (!TestNotNull("Function", Function)) return {};
-	auto ToolCall = MakeShared<UE::ToolsetRegistry::FObjectFunctionToolCall>(Toolset, Function);
-	auto Handler = MakeShared<UE::ToolsetRegistry::FToolCallExceptionHandler>();
-	ToolCall->Execute(
-		UE::ToolsetRegistry::FObjectFunctionToolCall::FFunctionInputParamsJson(
-			TInPlaceType<FString>(), ParamsJson),
-		Handler);
-	return Handler->GetException();
-}
+#define TestLogCategory LogToolsetRegistry
 
 void FLogsToolsetSpec::Define()
 {
+	using namespace UE::ToolsetRegistry;
+
 	Describe(TEXT("GetLogCategories"), [this]()
 	{
 		It(TEXT("Returns a non-empty sorted list"), [this]()
@@ -43,11 +42,11 @@ void FLogsToolsetSpec::Define()
 			GLog->Flush();
 
 			TArray<FString> Categories = ULogsToolset::GetLogCategories();
-			TestTrue("Categories is non-empty", Categories.Num() > 0);
+			TestGreaterThan(TEXT("Categories is non-empty"), Categories.Num(), 0);
 
 			TArray<FString> Sorted = Categories;
 			Sorted.Sort();
-			TestEqual("Categories is sorted", Categories, Sorted);
+			TestEqual(TEXT("Categories is sorted"), Categories, Sorted);
 		});
 
 		It(TEXT("Includes the test log category"), [this]()
@@ -56,7 +55,10 @@ void FLogsToolsetSpec::Define()
 			GLog->Flush();
 
 			TArray<FString> Categories = ULogsToolset::GetLogCategories();
-			TestTrue("LogLogsToolsetTest is present",
+			TestTrue(
+				FString::Printf(
+					TEXT("LogLogsToolsetTest is present in [%s]"),
+					*FString::Join(Categories, TEXT(", "))),
 				Categories.Contains(TEXT("LogLogsToolsetTest")));
 		});
 
@@ -66,11 +68,12 @@ void FLogsToolsetSpec::Define()
 			GLog->Flush();
 
 			TArray<FString> Filtered = ULogsToolset::GetLogCategories(TEXT("LogLogsToolsetTest"));
-			TestTrue("At least one result", Filtered.Num() > 0);
+			TestGreaterThan(TEXT("At least one result"), Filtered.Num(), 0);
 			for (const FString& Category : Filtered)
 			{
-				TestTrue("Category contains filter string",
-					Category.Contains(TEXT("LogLogsToolsetTest"), ESearchCase::IgnoreCase));
+				TestContains(
+					TEXT("Category contains filter string"),
+					Category, TEXT("LogLogsToolsetTest"), ESearchCase::IgnoreCase);
 			}
 		});
 	});
@@ -80,7 +83,7 @@ void FLogsToolsetSpec::Define()
 		It(TEXT("Returns a non-empty list with no filters"), [this]()
 		{
 			TArray<FString> Entries = ULogsToolset::GetLogEntries(TEXT(""), TEXT(""), 0);
-			TestTrue("Entries is non-empty", Entries.Num() > 0);
+			TestGreaterThan(TEXT("Entries is non-empty"), Entries.Num(), 0);
 		});
 
 		It(TEXT("Filters by category"), [this]()
@@ -90,11 +93,11 @@ void FLogsToolsetSpec::Define()
 
 			TArray<FString> Entries = ULogsToolset::GetLogEntries(
 				TEXT("LogLogsToolsetTest"), TEXT(""), 0);
-			TestTrue("At least one entry", Entries.Num() > 0);
+			TestGreaterThan(TEXT("At least one entry"), Entries.Num(), 0);
 			for (const FString& Entry : Entries)
 			{
-				TestTrue("Entry contains category prefix",
-					Entry.Contains(TEXT("]LogLogsToolsetTest:")));
+				TestContains(
+					TEXT("Entry contains category prefix"), Entry, TEXT("]LogLogsToolsetTest:"));
 			}
 		});
 
@@ -106,9 +109,9 @@ void FLogsToolsetSpec::Define()
 			GLog->Flush();
 
 			TArray<FString> Entries = ULogsToolset::GetLogEntries(TEXT(""), *Marker, 0);
-			if (TestEqual("Exactly one entry matches", Entries.Num(), 1))
+			if (TestEqual(TEXT("Exactly one entry matches"), Entries.Num(), 1))
 			{
-				TestTrue("Entry contains marker", Entries[0].Contains(Marker));
+				TestContains(TEXT("Entry contains marker"), Entries[0], Marker);
 			}
 		});
 
@@ -121,27 +124,38 @@ void FLogsToolsetSpec::Define()
 
 			TArray<FString> Entries = ULogsToolset::GetLogEntries(
 				TEXT(""), *Marker.ToLower(), 0);
-			TestEqual("Case-insensitive match finds entry", Entries.Num(), 1);
+			TestEqual(TEXT("Case-insensitive match finds entry"), Entries.Num(), 1);
 		});
 
 		It(TEXT("Raises an error for an unknown category"), [this]()
 		{
-			TestTrue("Script error raised",
-				CallAndCaptureException(
-					TEXT("GetLogEntries"),
-					TEXT(R"({"Category": "LogThisCategoryDoesNotExist_XYZ123", "Pattern": "", "MaxEntries": 0})"))
-				.Contains(TEXT("not found")));
+			FToolCallExceptionHandler ExceptionHandler;
+			ExceptionHandler.CaptureErrorsIn([this]() -> void
+				{
+					ULogsToolset::GetLogEntries(
+						TEXT("LogThisCategoryDoesNotExist_XYZ123"),
+						TEXT(""),
+						0);
+				});
+			TestContains(
+				TEXT("GetLogEntries with invalid category raises error"),
+				*ExceptionHandler.GetException(), TEXT("not found"));
 		});
 
 		It(TEXT("Respects MaxEntries limit"), [this]()
 		{
 			constexpr int32 Limit = 5;
-			TArray<FString> All = ULogsToolset::GetLogEntries(TEXT(""), TEXT(""), 0);
-			TArray<FString> Limited = ULogsToolset::GetLogEntries(TEXT(""), TEXT(""), Limit);
-			TestTrue("At most MaxEntries returned", Limited.Num() <= Limit);
+			for (int32 i = 0; i < Limit * 2; ++i)
+			{
+				UE_LOGF(TestLogCategory, Verbose, "Log testing %d", i);
+			}
+			FString LogCategory = TestLogCategory.GetCategoryName().ToString();
+			TArray<FString> All = ULogsToolset::GetLogEntries(LogCategory, TEXT(""), 0);
+			TArray<FString> Limited = ULogsToolset::GetLogEntries(LogCategory, TEXT(""), Limit);
+			TestLessEqual(TEXT("At most MaxEntries returned"), Limited.Num(), Limit);
 			if (All.Num() >= Limit)
 			{
-				TestEqual("Returns the most recent entries",
+				TestEqual(TEXT("Returns the most recent entries"),
 					Limited, TArray<FString>(All.GetData() + All.Num() - Limit, Limit));
 			}
 		});
@@ -155,7 +169,7 @@ void FLogsToolsetSpec::Define()
 
 			TArray<FString> Entries = ULogsToolset::GetLogEntries(
 				TEXT("LogLogsToolsetTest"), *Marker, 0);
-			TestEqual("Exactly one combined match", Entries.Num(), 1);
+			TestEqual(TEXT("Exactly one combined match"), Entries.Num(), 1);
 		});
 	});
 
@@ -165,32 +179,38 @@ void FLogsToolsetSpec::Define()
 		{
 			UE_LOG(LogLogsToolsetTest, Log, TEXT("GetVerbosity probe"));
 			const FString Verbosity = ULogsToolset::GetVerbosity(TEXT("LogLogsToolsetTest"));
-			TestFalse("Verbosity is non-empty", Verbosity.IsEmpty());
+			TestNotEqual(TEXT("Verbosity is non-empty"), *Verbosity, TEXT(""));
 		});
 
 		It(TEXT("Returns empty string for an unknown category"), [this]()
 		{
 			const FString Verbosity =
 				ULogsToolset::GetVerbosity(TEXT("LogThisCategoryDoesNotExist_XYZ123"));
-			TestTrue("Returns empty string for unknown category", Verbosity.IsEmpty());
+			TestEqual(TEXT("Returns empty string for unknown category"), Verbosity, TEXT(""));
 		});
 
 		It(TEXT("Raises an error for an unknown category"), [this]()
 		{
-			TestTrue("Script error raised",
-				CallAndCaptureException(
-					TEXT("GetVerbosity"),
-					TEXT(R"({"Category": "LogThisCategoryDoesNotExist_XYZ123"})"))
-				.Contains(TEXT("not found")));
+			FToolCallExceptionHandler ExceptionHandler;
+			ExceptionHandler.CaptureErrorsIn([this]() -> void
+				{
+					ULogsToolset::GetVerbosity(TEXT("LogThisCategoryDoesNotExist_XYZ123"));
+				});
+			TestContains(
+				TEXT("GetVerbosity raises error with an unknown category"),
+				ExceptionHandler.GetException(), TEXT("not found"));
 		});
 
 		It(TEXT("Raises an error for an empty category"), [this]()
 		{
-			TestTrue("Script error raised",
-				CallAndCaptureException(
-					TEXT("GetVerbosity"),
-					TEXT(R"({"Category": ""})"))
-				.Contains(TEXT("Category cannot be empty")));
+			FToolCallExceptionHandler ExceptionHandler;
+			ExceptionHandler.CaptureErrorsIn([this]() -> void
+				{
+					ULogsToolset::GetVerbosity(TEXT(""));
+				});
+			TestContains(
+				TEXT("GetVerbosity raises error with an empty category"),
+				ExceptionHandler.GetException(), TEXT("Category cannot be empty"));
 		});
 	});
 
@@ -200,36 +220,44 @@ void FLogsToolsetSpec::Define()
 		{
 			const FString Category = TEXT("LogLogsToolsetTest");
 			const FString Original = ULogsToolset::GetVerbosity(Category);
-			if (!TestFalse("Category verbosity is readable", Original.IsEmpty()))
+			if (!TestNotEqual(TEXT("Category verbosity is readable"), *Original, TEXT("")))
 			{
 				return;
 			}
 
 			ULogsToolset::SetVerbosity(Category, TEXT("Warning"));
-			TestEqual("Verbosity set to Warning",
+			TestEqual(TEXT("Verbosity set to Warning"),
 				ULogsToolset::GetVerbosity(Category), FString(TEXT("Warning")));
 
 			ULogsToolset::SetVerbosity(Category, Original);
-			TestEqual("Verbosity restored",
+			TestEqual(TEXT("Verbosity restored"),
 				ULogsToolset::GetVerbosity(Category), Original);
 		});
 
 		It(TEXT("Raises an error for an empty category"), [this]()
 		{
-			TestTrue("Script error raised",
-				CallAndCaptureException(
-					TEXT("SetVerbosity"),
-					TEXT(R"({"Category": "", "Verbosity": "Log"})"))
-				.Contains(TEXT("Category cannot be empty")));
+			FToolCallExceptionHandler ExceptionHandler;
+			ExceptionHandler.CaptureErrorsIn([this]() -> void
+				{
+					ULogsToolset::SetVerbosity(TEXT(""), TEXT("Display"));
+				});
+			TestContains(
+				TEXT("SetVerbosity raises error with an empty category"),
+				ExceptionHandler.GetException(), TEXT("Category cannot be empty"));
 		});
 
 		It(TEXT("Raises an error for an invalid verbosity"), [this]()
 		{
-			TestTrue("Script error raised",
-				CallAndCaptureException(
-					TEXT("SetVerbosity"),
-					TEXT(R"({"Category": "LogLogsToolsetTest", "Verbosity": "NotAVerbosity"})"))
-				.Contains(TEXT("is not a valid verbosity level")));
+			FToolCallExceptionHandler ExceptionHandler;
+			ExceptionHandler.CaptureErrorsIn([this]() -> void
+				{
+					ULogsToolset::SetVerbosity(
+						TestLogCategory.GetCategoryName().ToString(),
+						TEXT("NotAVerbosity"));
+				});
+			TestContains(
+				TEXT("SetVerbosity raises error with an invalid category"),
+				ExceptionHandler.GetException(), TEXT("is not a valid verbosity level"));
 		});
 	});
 }
