[Toolset Registry] Python test runner

Provides a Python interface to register Python unittest tests (though it could be expanded to other
test types) with the Unreal automation test framework.
`unreal.PythonTestRunner.create()` is should be used to create a `FAutomationTestBase` that discovers
and runs Python tests. The caller of `unreal.PythonTestRunner.create()` just needs to hold a global
reference to the returned object (e.g in their `init_unreal.py`).

#Synced-CL 52565678

diff --git a/Engine/Plugins/Experimental/ToolsetRegistry/Source/ToolsetRegistry/Private/Tests/PythonTestRunnerTest.cpp b/Engine/Plugins/Experimental/ToolsetRegistry/Source/ToolsetRegistry/Private/Tests/PythonTestRunnerTest.cpp
new file mode 100644
--- /dev/null
+++ b/Engine/Plugins/Experimental/ToolsetRegistry/Source/ToolsetRegistry/Private/Tests/PythonTestRunnerTest.cpp
@@ -0,0 +1,482 @@
+// Copyright Epic Games, Inc. All Rights Reserved. 
+
+#include "Containers/UnrealString.h"
+#include "Misc/AutomationTest.h"
+#include "Templates/SharedPointer.h"
+#include "Templates/UnrealTemplate.h"
+#include "UObject/Object.h"
+#include "UObject/ObjectPtr.h"
+
+#include "ToolsetRegistry/PythonTestRunner.h"
+
+#if WITH_AUTOMATION_TESTS
+
+class FPythonAutomationTestAccessor
+{
+public:
+	static TValueOrError<TArray<FString>, FString> ParseJsonStringArray(
+		const FString& StringArrayJson)
+	{
+		return UE::ToolsetRegistry::FPythonAutomationTest::ParseJsonStringArray(StringArrayJson);
+	}
+
+	static TValueOrError<FString, FString> RunTestRunnerPythonCode(
+		const FString& CodeToRun, const FString& ErrorContext)
+	{
+		return UE::ToolsetRegistry::FPythonAutomationTest::RunTestRunnerPythonCode(
+			CodeToRun, ErrorContext);
+	}
+
+	static TValueOrError<FString, FString> StartTestExecution(
+		UE::ToolsetRegistry::FPythonAutomationTest& Test, const FString& TestId)
+	{
+		return Test.StartTestExecution(TestId);
+	}
+
+	static bool PollForTestCompletion(
+		UE::ToolsetRegistry::FPythonAutomationTest& Test, const FString& TestId)
+	{
+		return Test.PollForTestCompletion(TestId);
+	}
+};
+
+class FPythonTestRunnerAccessor
+{
+public:
+	static void SetRunTestDelegate(
+		TStrongObjectPtr<UPythonTestRunner> Runner,
+		TFunction<void(const FString&)>&& RunTestDelegate)
+	{
+		Runner->RunTestDelegate = MoveTemp(RunTestDelegate);
+	}
+};
+
+BEGIN_DEFINE_SPEC(
+	FPythonTestRunnerTest,
+	"AI.ToolsetRegistry.PythonTestRunnerTest",
+	EAutomationTestFlags::EditorContext | EAutomationTestFlags::EngineFilter)
+private:
+	// Get search options for test cases used by unittest_runner.
+	static FPythonTestRunnerSearchOptions GetSearchOptions(const FString& TestNameGlob)
+	{
+		FPythonTestRunnerSearchOptions Options;
+		Options.RootModule = PythonTestsRootModule;
+		Options.TestModuleGlob = PythonTestModule;
+		Options.TestNameGlobs.Add(TestNameGlob);
+		return Options;
+	}
+
+	// Create a Python automation test that discovers unittest_runner test cases.
+	static TSharedPtr<UE::ToolsetRegistry::FPythonAutomationTest> CreatePythonAutomationTest(
+		const FString& TestNameGlob)
+	{
+		return MakeShared<UE::ToolsetRegistry::FPythonAutomationTest>(
+			TestBaseName, GetSearchOptions(TestNameGlob));
+	}
+
+	// Create a Python test runner.
+	TStrongObjectPtr<UPythonTestRunner> CreatePythonTestRunner(
+		const FString& TestNameGlob)
+	{
+		return TStrongObjectPtr<UPythonTestRunner>(
+			UPythonTestRunner::Create(TestBaseName, GetSearchOptions(TestNameGlob)));
+	}
+	
+private:
+	// Ensure that a result has an error and the error string is not empty.
+	template<typename ValueT>
+	bool HasError(const TValueOrError<ValueT, FString>& MaybeError)
+	{
+		if (!TestTrue(TEXT("HasError"), MaybeError.HasError())) return false;
+		return TestNotEqual(TEXT("ErrorString"), MaybeError.GetError(), TEXT(""));
+	}
+
+	// Clear last test results to reset any global state.
+	void ClearLastTestResults()
+	{
+		TestTrue(
+			TEXT("Clear results"),
+			FPythonAutomationTestAccessor::RunTestRunnerPythonCode(
+				TEXT("unittest_runner.clear_last_test_results()"),
+				TEXT("Setup test case")).HasValue());
+	}
+
+	// Expect a UPythonTestRunner not created error.
+	void ExpectPythonTestRunnerNotCreatedError()
+	{
+		AddExpectedError(
+			TEXT("must be instanced using create"),
+			EAutomationExpectedMessageFlags::MatchType::Contains, 1, true);
+	}
+
+	// Expect Python error logs.
+	void ExpectPythonErrors()
+	{
+		// Python logs multiple error lines so ignore them all.
+		AddExpectedError(
+			TEXT(".*"),
+			EAutomationExpectedMessageFlags::MatchType::Contains, 0, true);
+	}
+
+private:
+	static const TCHAR TestBaseName[];
+	static const TCHAR PythonTestsRootModule[];
+	static const TCHAR PythonTestModule[];
+	static const TCHAR PythonTestSuccessId[];
+	static const TCHAR PythonTestFailId[];
+	END_DEFINE_SPEC(FPythonTestRunnerTest)
+
+const TCHAR FPythonTestRunnerTest::TestBaseName[] = TEXT("PythonAutomationTest");
+const TCHAR FPythonTestRunnerTest::PythonTestsRootModule[] =
+	TEXT("toolset_registry.tests");
+const TCHAR FPythonTestRunnerTest::PythonTestModule[] =
+	TEXT("_unittest_runner_test_cases");
+const TCHAR FPythonTestRunnerTest::PythonTestSuccessId[] =
+	TEXT("toolset_registry.tests._unittest_runner_test_cases.TestSuccess.test_succeed");
+const TCHAR FPythonTestRunnerTest::PythonTestFailId[] =
+	TEXT("toolset_registry.tests._unittest_runner_test_cases.TestFailure.test_fail");
+
+void FPythonTestRunnerTest::Define()
+{
+	using namespace UE::ToolsetRegistry;
+
+	Describe(TEXT("FPythonTestRunnerSearchOptions"), [this]()
+	{
+		Describe(TEXT("BuildPythonCode"), [this]()
+		{
+			It(TEXT("Returns code for empty object"), [this]()
+			{
+				FPythonTestRunnerSearchOptions Options;
+				TValueOrError<FString, FString> MaybePythonCode = Options.BuildPythonCode();
+				if (!TestFalse(TEXT("HasError"), MaybePythonCode.HasError())) return;
+				TestEqual(
+					TEXT("Code"),
+					*MaybePythonCode.GetValue(),
+					TEXT("unittest_runner.TestSearchOptions()"));
+			});
+
+			It(TEXT("Returns code with a root module"), [this]()
+			{
+				FPythonTestRunnerSearchOptions Options;
+				Options.RootModule = "foo.bar";
+				TValueOrError<FString, FString> MaybePythonCode = Options.BuildPythonCode();
+				if (!TestFalse(TEXT("HasError"), MaybePythonCode.HasError())) return;
+				TestEqual(
+					TEXT("Code"),
+					*MaybePythonCode.GetValue(),
+					TEXT("unittest_runner.TestSearchOptions(root_module='foo.bar')"));
+			});
+
+			It(TEXT("Fails with invalid root module"), [this]()
+			{
+				FPythonTestRunnerSearchOptions Options;
+				Options.RootModule = "foo'bar";
+				HasError(Options.BuildPythonCode());
+			});
+
+			It(TEXT("Returns code with module glob"), [this]()
+			{
+				FPythonTestRunnerSearchOptions Options;
+				Options.TestModuleGlob = "test_things*";
+				TValueOrError<FString, FString> MaybePythonCode = Options.BuildPythonCode();
+				if (!TestFalse(TEXT("HasError"), MaybePythonCode.HasError())) return;
+				TestEqual(
+					TEXT("Code"),
+					*MaybePythonCode.GetValue(),
+					TEXT("unittest_runner.TestSearchOptions(test_module_glob='test_things*')"));
+			});
+
+			It(TEXT("Fails with invalid module glob"), [this]()
+			{
+				FPythonTestRunnerSearchOptions Options;
+				Options.TestModuleGlob = "'test_things'";
+				HasError(Options.BuildPythonCode());
+			});
+
+			It(TEXT("Returns code with test glob"), [this]()
+			{
+				FPythonTestRunnerSearchOptions Options;
+				Options.TestNameGlobs.Add("test_foo");
+				Options.TestNameGlobs.Add("test_bar");
+				TValueOrError<FString, FString> MaybePythonCode = Options.BuildPythonCode();
+				if (!TestFalse(TEXT("HasError"), MaybePythonCode.HasError())) return;
+				TestEqual(
+					TEXT("Code"),
+					*MaybePythonCode.GetValue(),
+					TEXT(
+						"unittest_runner.TestSearchOptions("
+						"test_name_globs=set(['test_foo', 'test_bar']))"));
+			});
+
+			It(TEXT("Fails with invalid test glob"), [this]()
+			{
+				FPythonTestRunnerSearchOptions Options;
+				Options.TestNameGlobs.Add("test_foo");
+				Options.TestNameGlobs.Add("test'bar");
+				TValueOrError<FString, FString> MaybePythonCode = Options.BuildPythonCode();
+				HasError(Options.BuildPythonCode());
+			});
+
+			It(TEXT("Returns code with multiple args"), [this]()
+			{
+				FPythonTestRunnerSearchOptions Options;
+				Options.RootModule = "foo.bar";
+				Options.TestModuleGlob = "test_*";
+				TValueOrError<FString, FString> MaybePythonCode = Options.BuildPythonCode();
+				if (!TestFalse(TEXT("HasError"), MaybePythonCode.HasError())) return;
+				TestEqual(
+					TEXT("Code"),
+					*MaybePythonCode.GetValue(),
+					TEXT(
+						"unittest_runner.TestSearchOptions("
+						"root_module='foo.bar', test_module_glob='test_*')"));
+			});
+		});
+	});
+
+	Describe(TEXT("FPythonAutomationTest"), [this]()
+	{
+		Describe(TEXT("ParseJsonStringArray"), [this]()
+		{
+			It(TEXT("Can parse a list of strings"), [this]()
+			{
+				TValueOrError<TArray<FString>, FString> MaybeResult =
+					FPythonAutomationTestAccessor::ParseJsonStringArray(
+						TEXT(R"json(["a", "bc"])json"));
+				if (!TestTrue(TEXT("Successful"), MaybeResult .HasValue())) return;
+
+				TArray<FString> Result = MaybeResult.StealValue();
+				TestEqual(TEXT("Number of items"), Result.Num(), 2);
+				TestEqual(TEXT("Items"), *FString::Join(Result, TEXT(", ")), TEXT("a, bc"));
+			});
+
+			It(TEXT("Fails to parse non-list"), [this]()
+			{
+				HasError(FPythonAutomationTestAccessor::ParseJsonStringArray(TEXT(R"json({})json")));
+			});
+
+			It(TEXT("Fails to parse list with non-strings"), [this]()
+			{
+				HasError(
+					FPythonAutomationTestAccessor::ParseJsonStringArray(
+						TEXT(R"json(["a", 2)json")));
+			});
+		});
+
+		Describe(TEXT("RunTestRunnerPythonCode"), [this]()
+		{
+			static const TCHAR ErrorContext[] = TEXT("Testing");
+
+			It(TEXT("Can run a valid expression"), [this]()
+			{
+				TValueOrError<FString, FString> MaybeResult =
+					FPythonAutomationTestAccessor::RunTestRunnerPythonCode(
+						TEXT("json.dumps(1)"), ErrorContext);
+				if (!TestTrue(TEXT("Successful"), MaybeResult .HasValue())) return;
+
+				TestEqual(TEXT("Result"), *MaybeResult.GetValue(), TEXT("1"));
+			});
+
+			It(TEXT("Fails with a invalid expression"), [this]()
+			{
+				ExpectPythonErrors();
+				TValueOrError<FString, FString> MaybeResult =
+					FPythonAutomationTestAccessor::RunTestRunnerPythonCode(
+						TEXT("1 = 2"), ErrorContext);
+				if (!TestTrue(TEXT("Fails"), MaybeResult.HasError())) return;
+				TestTrue(
+					TEXT("Error contains context"),
+					MaybeResult.GetError().Contains(ErrorContext));
+			});
+		});
+
+		Describe(TEXT("GetTests"), [this]
+		{
+			It(TEXT("Returns discovered tests"), [this]()
+			{
+				TSharedPtr<FPythonAutomationTest> AutomationTest =
+					CreatePythonAutomationTest(PythonTestSuccessId);
+				TArray<FString> TestIds;
+				TArray<FString> BeautifiedNames;
+				AutomationTest->GetTests(BeautifiedNames, TestIds);
+				if (!TestEqual(
+						TEXT("TestIds count matches name count"),
+						TestIds.Num(), BeautifiedNames.Num()))
+				{
+					return;
+				}
+
+				FString CommaSeparatedTestIds = FString::Join(TestIds, TEXT(", "));
+				TestEqual(TEXT("TestIds"), CommaSeparatedTestIds, PythonTestSuccessId);
+				TestEqual(
+					TEXT("TestIds match beautified names"),
+					CommaSeparatedTestIds,
+					FString::Join(BeautifiedNames, TEXT(", ")));
+			});
+		});
+
+		Describe(TEXT("TestExecution"), [this]()
+		{
+			static const TCHAR InvalidTestId[] = TEXT("invalid',test.id");
+			constexpr float TimeoutInSeconds = 10.0f;
+
+			BeforeEach([this]()
+			{
+				ClearLastTestResults();
+			});
+
+			Describe(TEXT("GetLastTestResult"), [this]()
+			{
+				It(TEXT("Returns nothing when a test isn't runnning"), [this]()
+				{
+					TSharedPtr<FPythonAutomationTest> AutomationTest =
+						CreatePythonAutomationTest(PythonTestSuccessId);
+					TValueOrError<TOptional<bool>, FString> MaybeResult =
+						AutomationTest->GetLastTestResult(PythonTestSuccessId);
+					if (!TestTrue(TEXT("Has result"), MaybeResult.HasValue())) return;
+					TestFalse(TEXT("Not running"), MaybeResult.GetValue().IsSet());
+				});
+
+				It(TEXT("Returns a result when a test is complete"), [this]()
+				{
+					TSharedPtr<FPythonAutomationTest> AutomationTest =
+						CreatePythonAutomationTest(PythonTestSuccessId);
+					TValueOrError<FString, FString> MaybeTestId =
+						FPythonAutomationTestAccessor::StartTestExecution(
+							*AutomationTest, PythonTestSuccessId);
+					if (!TestTrue(TEXT("Has test ID"), MaybeTestId.HasValue())) return;
+
+					AddCommand(
+						new FUntilCommand(
+							[this, AutomationTest, TestId = MaybeTestId.StealValue()]()
+							{
+								return FPythonAutomationTestAccessor::PollForTestCompletion(
+									*AutomationTest, TestId);
+							},
+							TimeoutInSeconds));
+				});
+
+				It(TEXT("Fails with an invalid test ID"), [this]()
+				{
+					TSharedPtr<FPythonAutomationTest> AutomationTest =
+						CreatePythonAutomationTest(PythonTestSuccessId);
+					HasError(AutomationTest->GetLastTestResult(InvalidTestId));
+				});
+
+				It(TEXT("Returns a failure when a test is fails"), [this]()
+				{
+					TSharedPtr<FPythonAutomationTest> AutomationTest =
+						CreatePythonAutomationTest(PythonTestFailId);
+					TValueOrError<FString, FString> MaybeTestId =
+						FPythonAutomationTestAccessor::StartTestExecution(
+							*AutomationTest, PythonTestFailId);
+					if (!TestTrue(TEXT("Has test ID"), MaybeTestId.HasValue())) return;
+
+					ExpectPythonErrors();
+					AddCommand(
+						new FUntilCommand(
+							[this, AutomationTest, TestId = MaybeTestId.StealValue()]()
+							{
+								auto MaybeResult = AutomationTest->GetLastTestResult(TestId);
+								if (MaybeResult.HasError()) return false;
+								if (!MaybeResult.GetValue().IsSet()) return false;
+								TestFalse(TEXT("Failed"), MaybeResult.GetValue().GetValue());
+								return true;
+							},
+							TimeoutInSeconds));
+				});
+			});
+
+			Describe(TEXT("RunTest"), [this]
+			{
+				It(TEXT("Can run a test to completion"), [this]()
+				{
+					TSharedPtr<FPythonAutomationTest> AutomationTest =
+						CreatePythonAutomationTest(PythonTestSuccessId);
+					AutomationTest->RunTest(PythonTestSuccessId);
+				});
+			});
+		});
+	});
+
+	Describe(TEXT("UPythonTestRunner"), [this]()
+	{
+		BeforeEach([this]()
+		{
+			ClearLastTestResults();
+		});
+
+		Describe(TEXT("GetTests"), [this]()
+		{
+			It(TEXT("Fails with a non-created instance"), [this]()
+			{
+				TStrongObjectPtr<UPythonTestRunner> Runner(NewObject<UPythonTestRunner>());
+				ExpectPythonTestRunnerNotCreatedError();
+				(void)Runner->GetTests();
+			});
+
+			It(TEXT("Returns discovered tests"), [this]()
+			{
+				auto Runner = CreatePythonTestRunner(PythonTestSuccessId);
+				TestEqual(
+					TEXT("TestIds"),
+					FString::Join(Runner->GetTests(), TEXT(", ")),
+					PythonTestSuccessId);
+			});
+		});
+
+		Describe(TEXT("TestExecution"), [this]()
+		{
+
+			Describe(TEXT("GetLastTestResult"), [this]()
+			{
+				It(TEXT("Fails with a non-created instance"), [this]()
+				{
+					TStrongObjectPtr<UPythonTestRunner> Runner(NewObject<UPythonTestRunner>());
+					ExpectPythonTestRunnerNotCreatedError();
+					(void)Runner->GetLastTestResult(PythonTestSuccessId);
+				});
+
+				It(TEXT("Return an incomplete result for a non-executed test"), [this]()
+				{
+					TStrongObjectPtr<UPythonTestRunner> Runner =
+						CreatePythonTestRunner(PythonTestSuccessId);
+					FPythonTestRunnerResult Result =
+						Runner->GetLastTestResult(PythonTestSuccessId);
+					TestFalse(TEXT("IsComplete"), Result.bIsComplete);
+					TestFalse(TEXT("IsSuccessful"), Result.bIsSuccessful);
+				});
+			});
+
+			Describe(TEXT("RunTest"), [this]()
+			{
+				It(TEXT("Fails with a non-created instance"), [this]()
+				{
+					TStrongObjectPtr<UPythonTestRunner> Runner(NewObject<UPythonTestRunner>());
+					ExpectPythonTestRunnerNotCreatedError();
+					Runner->RunTest(PythonTestSuccessId);
+				});
+
+				It(TEXT("Can run a test"), [this]()
+				{
+					TStrongObjectPtr<UPythonTestRunner> Runner =
+						CreatePythonTestRunner(PythonTestSuccessId);
+					FString TestNameToRun;
+					FPythonTestRunnerAccessor::SetRunTestDelegate(
+						Runner, 
+						[&TestNameToRun](const FString& ReceivedTestNameToRun)
+						{
+							TestNameToRun = ReceivedTestNameToRun;
+						});
+					Runner->RunTest(PythonTestSuccessId);
+					TestEqual(
+						TEXT("Test Name"),
+						TestNameToRun,
+						FString::Printf(TEXT("%s.%s"), TestBaseName, PythonTestSuccessId));
+				});
+			});
+		});
+	});
+}
+
+#endif  // WITH_AUTOMATION_TESTS
diff --git a/Engine/Plugins/Experimental/ToolsetRegistry/Source/ToolsetRegistry/Private/ToolsetRegistry/PythonTestRunner.cpp b/Engine/Plugins/Experimental/ToolsetRegistry/Source/ToolsetRegistry/Private/ToolsetRegistry/PythonTestRunner.cpp
new file mode 100644
--- /dev/null
+++ b/Engine/Plugins/Experimental/ToolsetRegistry/Source/ToolsetRegistry/Private/ToolsetRegistry/PythonTestRunner.cpp
@@ -0,0 +1,447 @@
+// Copyright Epic Games, Inc. All Rights Reserved. 
+
+#include "ToolsetRegistry/PythonTestRunner.h"
+
+
+#include "Algo/Find.h"
+#include "Algo/Transform.h"
+#include "Dom/JsonValue.h"
+#include "Engine/Engine.h"
+#include "IPythonScriptPlugin.h"
+#include "Misc/AssertionMacros.h"
+#include "Misc/CoreDelegates.h"
+#include "PythonScriptTypes.h"
+#include "UObject/Object.h"
+#include "UObject/ObjectPtr.h"
+#include "Serialization/JsonReader.h"
+#include "Serialization/JsonSerializer.h"
+#include "Serialization/JsonTypes.h"
+#include "Templates/UnrealTemplate.h"
+
+#include "ToolsetRegistry/Module.h"
+
+namespace
+{
+	// Get the initialized Python script plugin.
+	IPythonScriptPlugin* GetPythonScriptPlugin()
+	{
+		IPythonScriptPlugin* Python = IPythonScriptPlugin::Get();
+		if (!(Python && Python->IsPythonInitialized()))
+		{
+			return nullptr;
+		}
+		return Python;
+	}
+}
+
+TValueOrError<void, FString> ValidatePythonString(
+	const FString& StringToValidate, const FString& ErrorContext)
+{
+	static const TCHAR InvalidSubstring[] = TEXT("'");
+	if (StringToValidate.Contains(InvalidSubstring))
+	{
+		return MakeError(
+			FString::Printf(
+				TEXT("%s failed, invalid substring `%s` found in `%s`"),
+				*ErrorContext, InvalidSubstring, *StringToValidate));
+	}
+	return MakeValue();
+}
+
+TValueOrError<void, FString> ValidatePythonStrings(
+	const TArray<FString>& StringsToValidate, const FString& ErrorContext)
+{
+	for (const FString& StringToValidate : StringsToValidate)
+	{
+		TValueOrError<void, FString> MaybeError =
+			ValidatePythonString(StringToValidate, ErrorContext);
+		if (MaybeError.HasError()) return MaybeError;
+	}
+	return MakeValue();
+}
+
+TValueOrError<FString, FString> FPythonTestRunnerSearchOptions::BuildPythonCode() const
+{
+	
+	TFunction<TValueOrError<void, FString>()> Validators[] = {
+		[this]() -> auto
+		{
+			return ValidatePythonString(
+				RootModule, TEXT("Validating search option RootModule"));
+		},
+		[this]() -> auto
+		{
+			return ValidatePythonString(
+				TestModuleGlob, TEXT("Validating search option TestModuleGlob"));
+		},
+		[this]() -> auto
+		{
+			return ValidatePythonStrings(
+				TestNameGlobs, TEXT("Validating search option TestNameGlobs"));
+		},
+	};
+	for (int i = 0; i < UE_ARRAY_COUNT(Validators); ++i)
+	{
+		TValueOrError<void, FString> MaybeError = Validators[i]();
+		if (MaybeError.HasError()) return MakeError(MaybeError.StealError());
+	}
+
+	TArray<FString> Arguments;
+	if (!RootModule.IsEmpty())
+	{
+		Arguments.Push(FString::Printf(TEXT("root_module='%s'"), *RootModule));
+	}
+	if (!TestModuleGlob.IsEmpty())
+	{
+		Arguments.Push(
+			FString::Printf(TEXT("test_module_glob='%s'"), *TestModuleGlob));
+	}
+	if (!TestNameGlobs.IsEmpty())
+	{
+		TArray<FString> QuotedTestNameGlobs;
+		Algo::Transform(
+			TestNameGlobs, QuotedTestNameGlobs,
+			[](const FString& Item) -> FString { return FString::Printf(TEXT("'%s'"), *Item); });
+		FString Items = FString::Join(QuotedTestNameGlobs, TEXT(", "));
+		Arguments.Push(
+			FString::Printf(TEXT("test_name_globs=set([%s])"), *Items));
+	}
+	return MakeValue(
+		FString::Printf(
+			TEXT("unittest_runner.TestSearchOptions(%s)"),
+			*FString::Join(Arguments, TEXT(", "))));
+}
+
+namespace UE::ToolsetRegistry
+{
+	FPythonAutomationTest::FPythonAutomationTest(
+		const FString& Name, const FPythonTestRunnerSearchOptions& InSearchOptions,
+		float InTimeoutInSeconds) :
+		FAutomationTestBase(Name, false),
+		SearchOptions(InSearchOptions),
+		TimeoutInSeconds(InTimeoutInSeconds) {}
+	FPythonAutomationTest::~FPythonAutomationTest() {}
+
+	EAutomationTestFlags FPythonAutomationTest::GetTestFlags() const
+	{
+		return EAutomationTestFlags::EditorContext |
+			EAutomationTestFlags::ProductFilter |
+			EAutomationTestFlags::CriticalPriority;
+	}
+
+	FString FPythonAutomationTest::GetBeautifiedTestName() const { return GetTestName(); }
+
+	uint32 FPythonAutomationTest::GetRequiredDeviceNum() const { return 1; }
+	
+	void FPythonAutomationTest::GetTests(
+		TArray<FString>& OutBeautifiedNames,
+		TArray <FString>& OutTestCommands) const
+	{
+		FString ErrorContext =
+			FString::Printf(TEXT("Discovering %s tests"), *GetTestName());
+		TValueOrError<FString, FString> MaybeSearchOptionsCode = SearchOptions.BuildPythonCode();
+		if (MaybeSearchOptionsCode.HasError())
+		{
+			UE_LOGF(LogToolsetRegistry, Error, "%ls", *MaybeSearchOptionsCode.GetError());
+			return;
+		}
+		TValueOrError<TArray<FString>, FString> MaybeTestIds = RunAndParseTestCaseIds(
+			FString::Printf(
+				TEXT("unittest_runner.discover_tests(search_options=%s)"),
+				*MaybeSearchOptionsCode.GetValue()),
+			ErrorContext);
+		if (MaybeTestIds.HasError())
+		{
+			UE_LOGF(LogToolsetRegistry, Error, "%ls", *MaybeTestIds.GetError());
+			return;
+		}
+		TArray<FString> TestIds = MaybeTestIds.StealValue();
+		if (TestIds.Num() == 0)
+		{
+			UE_LOGF(
+				LogToolsetRegistry, Error, "%ls failed, no tests found for search %ls",
+				*ErrorContext, *MaybeSearchOptionsCode.GetValue());
+			return;
+		}
+		for (const FString& TestId : TestIds)
+		{
+			OutBeautifiedNames.Add(TestId);
+			OutTestCommands.Add(TestId);
+		}
+	}
+
+	bool FPythonAutomationTest::RunTest(const FString& Parameters)
+	{
+		auto MaybeTestId = StartTestExecution(Parameters);
+		if (MaybeTestId.HasError())
+		{	
+			UE_LOGF(LogToolsetRegistry, Error, "%ls", *MaybeTestId.GetError());
+			return false;
+		}
+
+		AddCommand(
+			new FUntilCommand(
+				[This = AsShared(), TestId = MaybeTestId.StealValue()]() -> bool
+				{
+					return This->PollForTestCompletion(TestId);
+				},
+				TimeoutInSeconds));
+		return true;
+	}
+
+	TValueOrError<FString, FString> FPythonAutomationTest::StartTestExecution(
+		const FString& TestId)
+	{
+		FPythonTestRunnerSearchOptions SingleTestSearchOptions = SearchOptions;
+		SingleTestSearchOptions.TestNameGlobs.Empty();
+		SingleTestSearchOptions.TestNameGlobs.Add(TestId);
+
+		FString ErrorContext =
+			FString::Printf(TEXT("Running %s test %s"), *GetTestName(), *TestId);
+		TValueOrError<FString, FString> MaybeSearchOptionsCode =
+			SingleTestSearchOptions.BuildPythonCode();
+		if (MaybeSearchOptionsCode.HasError())
+		{
+			return MakeError(MaybeSearchOptionsCode.StealError());
+		}
+		TValueOrError<TArray<FString>, FString> MaybeTestIds = RunAndParseTestCaseIds(
+			FString::Printf(
+				TEXT(
+					"unittest_runner.discover_and_run_tests(\n"
+					"    search_options=%s,\n"
+					"    run_options=unittest_runner.TestRunOptions(\n"
+					"        use_automation_scheduler=True\n"
+					"    )\n"
+					")\n"),
+				*MaybeSearchOptionsCode.GetValue()),
+			ErrorContext);
+
+		if (MaybeTestIds.HasError()) return MakeError(MaybeTestIds.StealError());
+
+		TArray<FString> TestIds = MaybeTestIds.StealValue();
+		if (TestIds.Num() != 1)
+		{
+			return MakeError(
+				FString::Printf(
+					TEXT("%s failed. Unexpected number of tests %d when searching with %s."),
+					*ErrorContext, TestIds.Num(), *MaybeSearchOptionsCode.GetValue()));
+		}
+		return MakeValue(TestIds[0]);
+	}
+
+	bool FPythonAutomationTest::PollForTestCompletion(const FString& TestId)
+	{
+		// NOTE: When a test is executed via RunTest(), this method is called in the context
+		// of FUntilCommand with a timeout configured.
+		TValueOrError<TOptional<bool>, FString> MaybeResult = GetLastTestResult(TestId);
+		if (MaybeResult.HasError())
+		{
+			UE_LOGF(LogToolsetRegistry, Error, "%ls", *MaybeResult.GetError());
+			return true /* stop polling */;
+		}
+		check(MaybeResult.HasValue());
+		TOptional<bool> Result = MaybeResult.StealValue();
+		if (Result.IsSet())
+		{
+			TestTrue(TEXT("Test succeeded"), Result.GetValue());
+			return true /* stop polling */;
+		}
+		return false; /* continue polling */;
+	}
+
+	TValueOrError<TOptional<bool>, FString> FPythonAutomationTest::GetLastTestResult(
+		const FString& TestId) const
+	{
+		const FString ErrorContext =
+			FString::Printf(TEXT("Getting %s test %s result"), *GetTestName(), *TestId);
+		TValueOrError<void, FString> MaybeTestIdError =
+			ValidatePythonString(TestId, TEXT("Validating test ID"));
+		if (MaybeTestIdError.HasError()) return MakeError(MaybeTestIdError.StealError());
+
+		TValueOrError<FString, FString> StringResult = RunTestRunnerPythonCode(
+			FString::Printf(
+				TEXT("unittest_runner.get_last_test_result('%s')"),
+				*TestId),
+			ErrorContext);
+		if (StringResult.HasError()) return MakeError(StringResult.StealError());
+
+		TOptional<bool> Result;
+		FString StringResultValue = StringResult.StealValue();
+		// Still executing?
+		if (StringResultValue != TEXT("None")) Result.Emplace(StringResultValue.ToBool());
+		return MakeValue(Result);
+	}
+
+	TValueOrError<TArray<FString>, FString> FPythonAutomationTest::ParseJsonStringArray(
+		const FString& StringArrayJson)
+	{
+		TArray<TSharedPtr<FJsonValue>> JsonParsed;
+		TSharedRef<TJsonReader<>> Reader = TJsonReaderFactory<>::Create(StringArrayJson);
+		if (!FJsonSerializer::Deserialize(Reader, JsonParsed))
+		{
+			return MakeError(
+				FString::Printf(TEXT("Failed to parse JSON list '%s'"), *StringArrayJson));
+		}
+		TArray<FString> Result;
+		for (const auto& Value : JsonParsed)
+		{
+			if (Value->Type != EJson::String)
+			{
+				return MakeError(
+					FString::Printf(
+						TEXT("Unexpected non-string value found in JSON array '%s'"),
+						*StringArrayJson));
+			}
+			Result.Add(Value->AsString());
+		}
+		return MakeValue(MoveTemp(Result));
+    }
+
+	TValueOrError<FString, FString> FPythonAutomationTest::RunTestRunnerPythonCode(
+		const FString& CodeToRun, const FString& ErrorContext)
+	{
+		IPythonScriptPlugin* PythonScriptPlugin = GetPythonScriptPlugin();
+		if (!PythonScriptPlugin)
+		{
+			return MakeError(
+				FString::Printf(TEXT("%s failed, Python not initialized"), *ErrorContext));
+		}
+
+		FPythonCommandEx Command;
+		Command.Command =
+			FString::Printf(
+				TEXT(
+					"(exec('import json; from toolset_registry.tests import unittest_runner') "
+					"or %s)"),
+				*CodeToRun);
+		Command.ExecutionMode = EPythonCommandExecutionMode::EvaluateStatement;
+		if (!PythonScriptPlugin->ExecPythonCommandEx(Command))
+		{
+			return MakeError(
+				FString::Printf(
+					TEXT("%s failed, Python failed to execute '%s'"),
+					*ErrorContext, *Command.Command));
+		}
+		// The Python Script plugin uses repr to convert strings to strings which quotes them with
+		// single quotes, so remove quotation here.
+		bool bIgnoredCharRemoved;
+		Command.CommandResult.TrimCharInline(FString::ElementType('\''), &bIgnoredCharRemoved);
+		return MakeValue(Command.CommandResult);
+	}
+
+	TValueOrError<TArray<FString>, FString> FPythonAutomationTest::RunAndParseTestCaseIds(
+		const FString& FunctionThatReturnsTestCaseSequence,
+		const FString& ErrorContext)
+	{
+		TValueOrError<FString, FString> Result =
+			RunTestRunnerPythonCode(
+				FString::Printf(
+					TEXT("json.dumps([test.id() for test in %s])"),
+					*FunctionThatReturnsTestCaseSequence),
+				ErrorContext);
+		if (Result.HasError()) return MakeError(Result.StealError());
+		return ParseJsonStringArray(Result.StealValue());
+	}
+}
+
+UPythonTestRunner::~UPythonTestRunner()
+{
+	Reset();
+}
+
+TArray<FString> UPythonTestRunner::GetTests() const
+{
+	TArray<FString> IgnoredNames;
+	TArray<FString> TestIds;
+	if (CheckInitialized()) Test->GetTests(IgnoredNames, TestIds);
+	return TestIds;
+}
+
+void UPythonTestRunner::RunTest(const FString& TestId)
+{
+	if (CheckInitialized())
+	{
+		check(RunTestDelegate.IsSet());
+		RunTestDelegate(FString::Printf(TEXT("%s.%s"), *Test->GetTestName(), *TestId));
+	}
+}
+
+FPythonTestRunnerResult UPythonTestRunner::GetLastTestResult(const FString& TestId) const
+{
+	FPythonTestRunnerResult TestResult = { false, false };
+	if (!CheckInitialized()) return TestResult;
+
+	TValueOrError<TOptional<bool>, FString> MaybeResult = Test->GetLastTestResult(TestId);
+	if (MaybeResult.HasError())
+	{
+		UE_LOGF(LogToolsetRegistry, Error, "%ls", *MaybeResult.GetError());
+		return TestResult;
+	}
+
+	TOptional<bool> Result = MaybeResult.StealValue();
+	if (!Result.IsSet()) return TestResult;
+	TestResult.bIsComplete = true;
+	TestResult.bIsSuccessful = Result.GetValue();
+	return TestResult;
+}
+
+bool UPythonTestRunner::CheckInitialized() const
+{
+	if (!Test.IsValid())
+	{
+		UE_LOGF(LogToolsetRegistry, Error, "PythonTestRunner must be instanced using create()");
+		return false;
+	}
+	return true;
+}
+
+void UPythonTestRunner::Initialize(
+	const FString& BaseName,
+	const FPythonTestRunnerSearchOptions& SearchOptions,
+	float TimeoutInSeconds)
+{
+	Reset();
+	Test = MakeShared<UE::ToolsetRegistry::FPythonAutomationTest>(
+		BaseName, SearchOptions, TimeoutInSeconds);
+	FAutomationTestFramework::Get().RegisterAutomationTest(Test->GetTestName(), Test.Get());
+	RunTestDelegate =
+		[](const FString& TestName) -> void
+		{
+			if (GEngine)
+			{
+				GEngine->Exec(
+					GEngine->GetWorld(),
+					*FString::Printf(TEXT("Automation RunTests %s"), *TestName));
+			}
+			else
+			{
+				UE_LOGF(
+					LogToolsetRegistry, Error, "Unable to run test %ls, engine not available",
+					*TestName);
+			}
+		};
+	OnExitDelegateHandle = UE::ToolsetRegistry::FDelegateHandleRaii::Create(
+		FCoreDelegates::OnEnginePreExit,
+		FCoreDelegates::OnEnginePreExit.AddUObject(this, &UPythonTestRunner::Reset));
+}
+
+void UPythonTestRunner::Reset()
+{
+	if (Test.IsValid())
+	{
+		FAutomationTestFramework::Get().UnregisterAutomationTest(Test->GetTestName());
+		Test.Reset();
+	}
+	RunTestDelegate.Reset();
+	OnExitDelegateHandle.Reset();
+}
+
+UPythonTestRunner* UPythonTestRunner::Create(
+	const FString& BaseName,
+	const FPythonTestRunnerSearchOptions& SearchOptions,
+	float TimeoutInSeconds)
+{
+	TObjectPtr<UPythonTestRunner> NewInstance = NewObject<UPythonTestRunner>();
+	NewInstance->Initialize(BaseName, SearchOptions, TimeoutInSeconds);
+	return NewInstance.Get();
+}
diff --git a/Engine/Plugins/Experimental/ToolsetRegistry/Source/ToolsetRegistry/Private/ToolsetRegistry/PythonTestRunner.h b/Engine/Plugins/Experimental/ToolsetRegistry/Source/ToolsetRegistry/Private/ToolsetRegistry/PythonTestRunner.h
new file mode 100644
--- /dev/null
+++ b/Engine/Plugins/Experimental/ToolsetRegistry/Source/ToolsetRegistry/Private/ToolsetRegistry/PythonTestRunner.h
@@ -0,0 +1,214 @@
+// Copyright Epic Games, Inc. All Rights Reserved. 
+
+#pragma once
+
+#include "Containers/Array.h"
+#include "Containers/UnrealString.h"
+#include "Misc/AutomationTest.h"
+#include "Misc/Optional.h"
+#include "Templates/Function.h"
+#include "Templates/SharedPointer.h"
+#include "Templates/ValueOrError.h"
+#include "UObject/Object.h"
+#include "UObject/ObjectMacros.h"
+
+#include "ToolsetRegistry/DelegateHandle.h"
+
+#include "PythonTestRunner.generated.h"
+
+/// Search parameters for Python tests.
+USTRUCT(BlueprintType, Category="PythonTestRunner")
+struct FPythonTestRunnerSearchOptions
+{
+	GENERATED_BODY()
+
+	/// See toolset_registry.tests.unittest_runner.TestSearchOptions
+	/// Python module / package to search under.
+	UPROPERTY(EditAnywhere, BlueprintReadWrite, Category="PythonTestRunner")
+	FString RootModule;
+
+	/// Glob used to match test module names.
+	UPROPERTY(EditAnywhere, BlueprintReadWrite, Category="PythonTestRunner")
+	FString TestModuleGlob;
+
+	/// Globs used to match individual test names.
+	UPROPERTY(EditAnywhere, BlueprintReadWrite, Category="PythonTestRunner")
+	TArray<FString> TestNameGlobs;
+
+	/// Generate code to construct toolset_registry.tests.unittest_runner.TestSearchOptions
+	/// assuming test_runner is imported in the current scope.
+	TValueOrError<FString, FString> BuildPythonCode() const;
+};
+
+// Used to access internals of FPythonAutomationTest in tests.
+class FPythonAutomationTestAccessor;
+
+namespace UE::ToolsetRegistry
+{
+	/// Discovers and runs Python tests.
+	///
+	/// At the moment this has the limitation that all Python tests with class or module level
+	/// setup and teardown hooks (e.g setUpClass()) will be executed before and after every
+	/// test case.
+	class FPythonAutomationTest :
+		public FAutomationTestBase,
+		public TSharedFromThis<FPythonAutomationTest>
+	{
+		friend FPythonAutomationTestAccessor;
+
+	public:
+		/// Construct a Python automation test.
+		/// 
+		/// @param Name Base name of the test case (e.g My.Tests) which will be prefixed
+		///   to all Python test case names in test automation view.
+		/// @param SearchOptions Options used to discover Python test cases.
+		/// @param TimeoutInSeconds Time to wait for each test case to execute.
+		FPythonAutomationTest(
+			const FString& Name,
+			const FPythonTestRunnerSearchOptions& SearchOptions,
+			float TimeoutInSeconds = 5.0f);
+		virtual ~FPythonAutomationTest() override;
+
+		virtual EAutomationTestFlags GetTestFlags() const override;
+		virtual FString GetBeautifiedTestName() const override;
+		virtual uint32 GetRequiredDeviceNum() const override;
+
+		/// Discover Python tests.
+		///
+		/// @param OutBeautifiedNames Same as OutTestCommands.
+		/// @param OutTestCommands Fully qualified names of Python tests discovered by this
+		///   class.
+		virtual void GetTests(
+			TArray<FString>& OutBeautifiedNames,
+			TArray<FString>& OutTestCommands) const override;
+
+		/// Run requested Python test.
+		///
+		/// @param Parameters Test command to execute (see GetTests()).
+		///
+		/// @returns true if the test started, false otherwise.
+		virtual bool RunTest(const FString& Parameters) override;
+
+		/// Determine whether a test has finished executing after running it.
+		///
+		/// @param TestId ID of the test to query.
+		///
+		/// @returns If successful, returns an optional bool value that indicates whether the test
+		///   is not started or running (not set), succeeded or failed.
+		TValueOrError<TOptional<bool>, FString> GetLastTestResult(const FString& TestId) const;
+
+	private:
+		// Run a test without waiting for it to complete returning the ID of the test being
+		// executed.
+		TValueOrError<FString, FString> StartTestExecution(const FString& TestId);
+
+		// Poll to determine whether a test is complete returning false to continue polling, true
+		// to stop polling.
+		bool PollForTestCompletion(const FString& TestId);
+
+	private:
+		FPythonTestRunnerSearchOptions SearchOptions;
+		float TimeoutInSeconds;
+
+	private:
+		// Parse a JSON string list into a string array.
+		static TValueOrError<TArray<FString>, FString> ParseJsonStringArray(
+			const FString& StringArrayJson);
+
+		// Run Python code in a context where json and unittest_runner modules are imported.
+		static TValueOrError<FString, FString> RunTestRunnerPythonCode(
+			const FString& CodeToRun, const FString& ErrorContext);
+
+		// Run a function that returns unittest.TestCases returning IDs of the test cases.
+		// In this context json and unittest_runner modules are imported.
+		static TValueOrError<TArray<FString>, FString> RunAndParseTestCaseIds(
+			const FString& FunctionThatReturnsTestCaseSequence,
+			const FString& ErrorContext);
+	};
+}
+
+/// Result of a Python test.
+USTRUCT(BlueprintType, Category="PythonTestRunner")
+struct FPythonTestRunnerResult
+{
+	GENERATED_BODY()
+
+	/// Whether a queried test has been executed and is complete.
+	UPROPERTY(EditAnywhere, BlueprintReadWrite, Category="PythonTestRunner")
+	bool bIsComplete = false;
+
+	/// Whether a queried test completed successfully.
+	UPROPERTY(EditAnywhere, BlueprintReadWrite, Category="PythonTestRunner")
+	bool bIsSuccessful = false;
+};
+
+/// Python automation test runner.
+UCLASS(MinimalAPI, BlueprintType, Category="PythonTestRunner")
+class UPythonTestRunner : public UObject
+{
+	// Used to access internals of this class in tests.
+	friend class FPythonTestRunnerAccessor;
+
+	GENERATED_BODY()
+
+public:
+	virtual ~UPythonTestRunner() override;
+
+	/// Discover tests.
+	///
+	/// @returns Test IDs discovered by this runner.
+	UFUNCTION(BlueprintCallable, Category="PythonTestRunner")
+	TArray<FString> GetTests() const;
+
+	/// Run a test.
+	///
+	/// @param TestId Python test ID to run. This must be a test ID returned by
+	///    GetTests().
+	UFUNCTION(BlueprintCallable, Category="PythonTestRunner")
+	void RunTest(const FString& TestId);
+
+	/// Get the result of a test.
+	///
+	/// @param TestId Python test ID to query.
+	///
+	/// @return Last result of the test.
+	UFUNCTION(BlueprintCallable, Category="PythonTestRunner")
+	FPythonTestRunnerResult GetLastTestResult(const FString& TestId) const;
+
+private:
+	// Ensure the instance is initialize otherwise log an error and return false.
+	bool CheckInitialized() const;
+
+	// Initialize the instance.
+	void Initialize(
+		const FString& BaseName,
+		const FPythonTestRunnerSearchOptions& SearchOptions,
+		float TimeoutInSeconds);
+
+	// Clear the instance.
+	void Reset();
+
+private:
+	TSharedPtr<UE::ToolsetRegistry::FPythonAutomationTest> Test;
+	// Delegate that runs a test by name, this is overridable by tests so that it's possible
+	// to test RunTests() since tests can't recursively call
+	// FAutomationTestFramework::StartTestByName().
+	TFunction<void(const FString&)> RunTestDelegate;
+
+	UE::ToolsetRegistry::FDelegateHandleRaii OnExitDelegateHandle;
+
+public:
+	/// Create a test runner.
+	///
+	/// @param BaseName Base name of the test case (e.g My.Tests) which will be prefixed
+	///   to all Python test case names in test automation view.
+	/// @param SearchOptions Options used to discover Python test cases.
+	/// @param TimeoutInSeconds Time to wait for each test case to execute.
+	/// 
+	/// @returns Instance of this object that holds a reference to a FPythonAutomationTest.
+	UFUNCTION(BlueprintCallable, Category="PythonTestRunner")
+	static UPythonTestRunner* Create(
+		const FString& BaseName,
+		const FPythonTestRunnerSearchOptions& SearchOptions,
+		float TimeoutInSeconds = 5.0f);
+};
