[AI Toolsets]: Improve docs on how to use the PythonTestRunner from Python.

#Synced-CL 52596940

diff --git a/Engine/Plugins/Experimental/ToolsetRegistry/Source/ToolsetRegistry/Private/ToolsetRegistry/PythonTestRunner.h b/Engine/Plugins/Experimental/ToolsetRegistry/Source/ToolsetRegistry/Private/ToolsetRegistry/PythonTestRunner.h
--- a/Engine/Plugins/Experimental/ToolsetRegistry/Source/ToolsetRegistry/Private/ToolsetRegistry/PythonTestRunner.h
+++ b/Engine/Plugins/Experimental/ToolsetRegistry/Source/ToolsetRegistry/Private/ToolsetRegistry/PythonTestRunner.h
@@ -206,6 +206,24 @@ public:
 	/// @param TimeoutInSeconds Time to wait for each test case to execute.
 	/// 
 	/// @returns Instance of this object that holds a reference to a FPythonAutomationTest.
+	///
+	/// This is intended to be used from Python to register a set of Python unittest tests with the
+	/// Unreal test runner.
+	///
+	/// For example, the following when called from a plugin's `init_unreal.py` script, will register
+	/// tests under the package `foo.bar.tests` as tests underneath `MyPackage.Magic`:
+	///
+	/// @code{.py}
+	/// import unreal
+	/// from foo.bar import tests
+	/// 
+	/// # It is important to reference the runner in a global to prevent it from being garbage
+	/// # collected when instanced from an init_unreal.py script.
+	/// tests._test_runner = unreal.PythonTestRunner.create(
+	///   'MyPackage.Magic',
+	///   unreal.PythonTestRunnerSearchOptions(root_module=tests.__package__))
+	/// @endcode
+	///
 	UFUNCTION(BlueprintCallable, Category="PythonTestRunner")
 	static UPythonTestRunner* Create(
 		const FString& BaseName,
