[AI Toolsets] Register Python toolset tests with UE's test runner.

NOTE: This filters core toolset tests into a separate category.

#Synced-CL 52598835

diff --git a/Engine/Plugins/Experimental/ToolsetRegistry/Source/ToolsetRegistry/Private/Tests/BlueprintLibraryToolsetTest.cpp b/Engine/Plugins/Experimental/ToolsetRegistry/Source/ToolsetRegistry/Private/Tests/BlueprintLibraryToolsetTest.cpp
--- a/Engine/Plugins/Experimental/ToolsetRegistry/Source/ToolsetRegistry/Private/Tests/BlueprintLibraryToolsetTest.cpp
+++ b/Engine/Plugins/Experimental/ToolsetRegistry/Source/ToolsetRegistry/Private/Tests/BlueprintLibraryToolsetTest.cpp
@@ -802,25 +802,5 @@ void FToolsetRegistryToolsetTestGetAllSchemas::Define()
 	});
 }
 
-// Run python tests
-IMPLEMENT_SIMPLE_AUTOMATION_TEST(
-	FBlueprintLibraryToolsetTestPythonTests,
-	"AI.ToolsetRegistry.BlueprintLibraryToolset.PythonTests",
-	ToolsetRegistryTest::Flags)
-
-bool FBlueprintLibraryToolsetTestPythonTests::RunTest(
-	const FString& UnusedParameters)
-{
-	if (!IPythonScriptPlugin::Get()->IsPythonAvailable())
-	{
-		return true;
-	}
-
-	FPythonCommandEx PythonCommand = FPythonCommandEx();
-	PythonCommand.ExecutionMode = EPythonCommandExecutionMode::ExecuteStatement;
-	PythonCommand.Command = TEXT("import toolset_registry.tests; toolset_registry.tests.run_tests()");
-	return IPythonScriptPlugin::Get()->ExecPythonCommandEx(PythonCommand);
-}
-
 #endif	// WITH_DEV_AUTOMATION_TESTS
 
