[AI Toolset]: Only register toolsets if the registry is available.

It's possible to run games in the development editor target which runs
Python without the editor, so handle the case where init_unreal.py
scripts are executed but the toolset registry is available.

#Synced-CL 52655924

diff --git a/Engine/Plugins/Experimental/ToolsetRegistry/Source/ToolsetRegistry/Private/ToolsetRegistry/UToolsetRegistry.cpp b/Engine/Plugins/Experimental/ToolsetRegistry/Source/ToolsetRegistry/Private/ToolsetRegistry/UToolsetRegistry.cpp
--- a/Engine/Plugins/Experimental/ToolsetRegistry/Source/ToolsetRegistry/Private/ToolsetRegistry/UToolsetRegistry.cpp
+++ b/Engine/Plugins/Experimental/ToolsetRegistry/Source/ToolsetRegistry/Private/ToolsetRegistry/UToolsetRegistry.cpp
@@ -13,6 +13,11 @@
 #include "ToolsetRegistry/ToolsetRegistry.h"
 #include "ToolsetRegistry/ToolsetRegistrySubsystem.h"
 
+bool UToolsetRegistry::IsAvailable()
+{
+	return UToolsetRegistrySubsystem::Get(TEXT("IsAvailable")).HasValue();
+}
+
 void UToolsetRegistry::RegisterToolsetClass(TSubclassOf<UToolsetDefinition> InToolsetClass)
 {
 	if (!InToolsetClass)
diff --git a/Engine/Plugins/Experimental/ToolsetRegistry/Source/ToolsetRegistry/Public/ToolsetRegistry/UToolsetRegistry.h b/Engine/Plugins/Experimental/ToolsetRegistry/Source/ToolsetRegistry/Public/ToolsetRegistry/UToolsetRegistry.h
--- a/Engine/Plugins/Experimental/ToolsetRegistry/Source/ToolsetRegistry/Public/ToolsetRegistry/UToolsetRegistry.h
+++ b/Engine/Plugins/Experimental/ToolsetRegistry/Source/ToolsetRegistry/Public/ToolsetRegistry/UToolsetRegistry.h
@@ -25,6 +25,11 @@ class UToolsetRegistry : public UBlueprintFunctionLibrary
 	GENERATED_BODY()
 
 public:
+	// Whether the toolset registry is available.
+	// For example, this will return false if the editor is not available.
+	UFUNCTION(BlueprintCallable, Category="ToolsetRegistry")
+	static UE_API bool IsAvailable();
+
 	// Register a Blueprint function library class as a toolset.
 	UFUNCTION(BlueprintCallable, Category="ToolsetRegistry")
 	static UE_API void RegisterToolsetClass(TSubclassOf<UToolsetDefinition> InToolsetClass);
