Update CaptureAssetImage to have a ShowUI option, which defaults to false. When not showing UI, hide things like selection highlight and gizmos.

#Synced-CL 52671970

diff --git a/Engine/Plugins/Experimental/ToolsetRegistry/Source/ToolsetRegistry/Private/ToolsetRegistry/EditorAppToolset.cpp b/Engine/Plugins/Experimental/ToolsetRegistry/Source/ToolsetRegistry/Private/ToolsetRegistry/EditorAppToolset.cpp
--- a/Engine/Plugins/Experimental/ToolsetRegistry/Source/ToolsetRegistry/Private/ToolsetRegistry/EditorAppToolset.cpp
+++ b/Engine/Plugins/Experimental/ToolsetRegistry/Source/ToolsetRegistry/Private/ToolsetRegistry/EditorAppToolset.cpp
@@ -55,19 +55,29 @@ namespace
 	class FViewportScreenshotCapture : public TSharedFromThis<FViewportScreenshotCapture>
 	{
 	public:
-		FViewportScreenshotCapture() : Result(NewObject<UToolCallAsyncResultImage>()) {}
+		explicit FViewportScreenshotCapture(bool bInShowUI)
+			: Result(NewObject<UToolCallAsyncResultImage>())
+			, bShowUI(bInShowUI)
+		{}
 		~FViewportScreenshotCapture()
 		{
 			if (LevelEditorViewport.IsValid())
 			{
-				LevelEditorViewport->GetLevelViewportClient().RemoveRealtimeOverride(
-					FText::FromString(Label));
+				FLevelEditorViewportClient& ViewportClient =
+					LevelEditorViewport->GetLevelViewportClient();
+				ViewportClient.RemoveRealtimeOverride(FText::FromString(Label));
+				if (!bShowUI)
+				{
+					ViewportClient.EngineShowFlags.SetModeWidgets(bSavedModeWidgets);
+					ViewportClient.EngineShowFlags.SetSelectionOutline(bSavedSelectionOutline);
+					ViewportClient.EngineShowFlags.SetSelection(bSavedSelection);
+				}
 			}
 		}
 
-		static TStrongObjectPtr<UToolCallAsyncResultImage> Start()
+		static TStrongObjectPtr<UToolCallAsyncResultImage> Start(bool bShowUI)
 		{
-			return MakeShared<FViewportScreenshotCapture>()->StartCapture();
+			return MakeShared<FViewportScreenshotCapture>(bShowUI)->StartCapture();
 		}
 
 	private:
@@ -90,7 +100,7 @@ namespace
 							This->OnCaptured(Width, Height, Colors);
 							This.Reset();
 						}));
-			FScreenshotRequest::RequestScreenshot(false);
+			FScreenshotRequest::RequestScreenshot(bShowUI);
 			return Result;
 		}
 
@@ -107,8 +117,20 @@ namespace
 			LevelEditorViewport = LevelEditor->GetFirstActiveLevelViewport();
 			if (!LevelEditorViewport) return false;
 
-			LevelEditorViewport->GetLevelViewportClient().AddRealtimeOverride(
-				true, FText::FromString(Label));
+			FLevelEditorViewportClient& ViewportClient =
+				LevelEditorViewport->GetLevelViewportClient();
+			ViewportClient.AddRealtimeOverride(true, FText::FromString(Label));
+
+			if (!bShowUI)
+			{
+				bSavedModeWidgets = ViewportClient.EngineShowFlags.ModeWidgets != 0;
+				bSavedSelectionOutline = ViewportClient.EngineShowFlags.SelectionOutline != 0;
+				bSavedSelection = ViewportClient.EngineShowFlags.Selection != 0;
+				ViewportClient.EngineShowFlags.SetModeWidgets(false);
+				ViewportClient.EngineShowFlags.SetSelectionOutline(false);
+				ViewportClient.EngineShowFlags.SetSelection(false);
+			}
+
 			return true;
 		}
 
@@ -150,6 +172,10 @@ namespace
 		TStrongObjectPtr<UToolCallAsyncResultImage> Result;
 		UE::ToolsetRegistry::FDelegateHandleRaii Handle;
 		TSharedPtr<SLevelViewport> LevelEditorViewport;
+		bool bShowUI = true;
+		bool bSavedModeWidgets = true;
+		bool bSavedSelectionOutline = true;
+		bool bSavedSelection = true;
 
 	private:
 		static constexpr TCHAR Label[] = TEXT("ViewportScreenshotCapture");
@@ -353,7 +379,7 @@ TSharedPtr<FJsonObject> UEditorAppToolset::CVarToJson(IConsoleObject* CVar)
 	return CVarData;
 }
 
-UToolCallAsyncResultImage* UEditorAppToolset::CaptureAssetImage(const FString& AssetPath)
+UToolCallAsyncResultImage* UEditorAppToolset::CaptureAssetImage(const FString& AssetPath, bool bShowUI)
 {
 	// An empty path or the currently loaded level to viewport screenshot.
 	UWorld* EditorWorld = GEditor ? GEditor->GetEditorWorldContext().World() : nullptr;
@@ -363,7 +389,7 @@ UToolCallAsyncResultImage* UEditorAppToolset::CaptureAssetImage(const FString& A
 	if (AssetPath.IsEmpty() ||
 		CurrentLevelPackage == FPackageName::ObjectPathToPackageName(AssetPath))
 	{
-		return FViewportScreenshotCapture::Start().Get();
+		return FViewportScreenshotCapture::Start(bShowUI).Get();
 	}
 
 	// Thumbnail rendering path
diff --git a/Engine/Plugins/Experimental/ToolsetRegistry/Source/ToolsetRegistry/Private/ToolsetRegistry/EditorAppToolset.h b/Engine/Plugins/Experimental/ToolsetRegistry/Source/ToolsetRegistry/Private/ToolsetRegistry/EditorAppToolset.h
--- a/Engine/Plugins/Experimental/ToolsetRegistry/Source/ToolsetRegistry/Private/ToolsetRegistry/EditorAppToolset.h
+++ b/Engine/Plugins/Experimental/ToolsetRegistry/Source/ToolsetRegistry/Private/ToolsetRegistry/EditorAppToolset.h
@@ -37,10 +37,13 @@ public:
 	 * montages, materials, textures), a thumbnail is rendered.
 	 * @param AssetPath The path to the asset, e.g. '/Game/Meshes/SM_Cube'.
 	 *   Pass empty string or the current level path to capture a viewport screenshot instead.
+	 * @param bShowUI If true, hides editor UI overlays such as transform gizmos and selection
+	 *   outlines when capturing the image.
 	 * @return An image of the asset.
 	 */
 	UFUNCTION(BlueprintCallable, Category = "EditorAppToolset")
-	static UToolCallAsyncResultImage* CaptureAssetImage(const FString& AssetPath);
+	static UToolCallAsyncResultImage* CaptureAssetImage(
+		const FString& AssetPath, bool bShowUI = false);
 
 	/*
 	 * Captures an image of the entire editor application as the user sees it.
