Fixup error messages on optional property in ReferenceConverter. Fixes issues where optional python properties show up named TransientPython..

#Synced-CL 52724283

diff --git a/Engine/Plugins/Experimental/ToolsetRegistry/Source/ToolsetRegistry/Private/ToolsetRegistry/ReferenceConverter.cpp b/Engine/Plugins/Experimental/ToolsetRegistry/Source/ToolsetRegistry/Private/ToolsetRegistry/ReferenceConverter.cpp
--- a/Engine/Plugins/Experimental/ToolsetRegistry/Source/ToolsetRegistry/Private/ToolsetRegistry/ReferenceConverter.cpp
+++ b/Engine/Plugins/Experimental/ToolsetRegistry/Source/ToolsetRegistry/Private/ToolsetRegistry/ReferenceConverter.cpp
@@ -21,6 +21,22 @@ namespace
 		return String.IsEmpty() || String.Equals("None", ESearchCase::IgnoreCase) || String.Equals("null", ESearchCase::IgnoreCase);
 	}
 
+	// Returns the display name for a property in error messages.
+	// Python optional params (e.g. 'unreal.Class | None') are implemented as an FOptionalProperty
+	// wrapping an inner transient value property named 'TransientPythonProperty'. In that case,
+	// return the outer optional property's name (the actual parameter name) instead.
+	FString GetPropertyDisplayName(const FProperty* Property)
+	{
+		if (const FField* OwnerField = Property->Owner.ToField())
+		{
+			if (CastField<const FOptionalProperty>(OwnerField))
+			{
+				return OwnerField->GetName();
+			}
+		}
+		return Property->GetName();
+	}
+
 	// Helper to validate soft reference type without loading the asset
 	bool ValidateSoftReferenceType(
 		const FString& RefPath,
@@ -241,7 +257,7 @@ namespace UE::ToolsetRegistry
 				if (const UFunction* OwningFunction = Cast<const UFunction>(Property->Owner.ToUObject()))
 				{
 					bool HasNullDefault = false;
-					FString MetadataKey = FString::Printf(TEXT("CPP_Default_%s"), *Property->GetName());
+					FString MetadataKey = FString::Printf(TEXT("CPP_Default_%s"), *GetPropertyDisplayName(Property));
 					if (OwningFunction->HasMetaData(*MetadataKey))
 					{
 						FString DefaultValue = OwningFunction->GetMetaData(*MetadataKey);
@@ -252,7 +268,7 @@ namespace UE::ToolsetRegistry
 					{
 						UKismetSystemLibrary::RaiseScriptError(
 							FString::Printf(TEXT("%s is not valid value for property '%s'"),
-											*RefPath, *Property->GetName()));
+											*RefPath, *GetPropertyDisplayName(Property)));
 						return true;
 					}
 				}
@@ -278,14 +294,14 @@ namespace UE::ToolsetRegistry
 				{
 					UKismetSystemLibrary::RaiseScriptError(
 						FString::Printf(TEXT("%s is not valid %s for property '%s'"),
-										*RefPath, *UClass::StaticClass()->GetName(), *Property->GetName()));
+										*RefPath, *UClass::StaticClass()->GetName(), *GetPropertyDisplayName(Property)));
 					return true;
 				}
 				else if (MetaClass && !Cast<UClass>(Object)->IsChildOf(MetaClass))
 				{
 					UKismetSystemLibrary::RaiseScriptError(
 						FString::Printf(TEXT("%s is not valid %s for property '%s'"),
-										*RefPath, *MetaClass->GetName(), *Property->GetName()));
+										*RefPath, *MetaClass->GetName(), *GetPropertyDisplayName(Property)));
 					return true;
 				}
 
@@ -322,7 +338,7 @@ namespace UE::ToolsetRegistry
 					{
 						UKismetSystemLibrary::RaiseScriptError(FString::Printf(
 							TEXT("%s is not a valid subclass of %s for property '%s'"),
-							*RefPath, *ObjectProperty->PropertyClass->GetName(), *Property->GetName()));
+							*RefPath, *ObjectProperty->PropertyClass->GetName(), *GetPropertyDisplayName(Property)));
 						return true;
 					}
 					// Back-calculate the outer: OutValue = (uint8*)Outer + Property->Offset
@@ -348,7 +364,7 @@ namespace UE::ToolsetRegistry
 						UKismetSystemLibrary::RaiseScriptError(
 							FString::Printf(
 								TEXT("%s is not valid %s for property '%s'"),
-								*RefPath, *ObjectProperty->PropertyClass->GetName(), *Property->GetName()));
+								*RefPath, *ObjectProperty->PropertyClass->GetName(), *GetPropertyDisplayName(Property)));
 						return true;
 					}
 
@@ -365,7 +381,7 @@ namespace UE::ToolsetRegistry
 				if (!ValidateSoftReferenceType(RefPath, SoftObjectProperty->PropertyClass, ErrorMessage))
 				{
 					UKismetSystemLibrary::RaiseScriptError(
-						FString::Printf(TEXT("%s '%s'"), *ErrorMessage, *Property->GetName()));
+						FString::Printf(TEXT("%s '%s'"), *ErrorMessage, *GetPropertyDisplayName(Property)));
 					return true;
 				}
 			}
@@ -377,7 +393,7 @@ namespace UE::ToolsetRegistry
 				if (!ValidateSoftReferenceType(RefPath, WeakObjectProperty->PropertyClass, ErrorMessage))
 				{
 					UKismetSystemLibrary::RaiseScriptError(
-						FString::Printf(TEXT("%s '%s'"), *ErrorMessage, *Property->GetName()));
+						FString::Printf(TEXT("%s '%s'"), *ErrorMessage, *GetPropertyDisplayName(Property)));
 					return true;
 				}
 			}
@@ -398,7 +414,7 @@ namespace UE::ToolsetRegistry
 					{
 						UE_LOGF(LogToolsetRegistry, Warning,
 							   "Soft reference path '%ls' not found in Asset Registry for property '%ls'",
-							   *RefPath, *Property->GetName());
+							   *RefPath, *GetPropertyDisplayName(Property));
 					}
 				}
 				else if (ScriptStruct == FSoftClassPath::StaticStruct())
@@ -417,7 +433,7 @@ namespace UE::ToolsetRegistry
 						{
 							UKismetSystemLibrary::RaiseScriptError(
 								FString::Printf(TEXT("%s is not a valid class for property '%s'"),
-												*RefPath, *Property->GetName()));
+												*RefPath, *GetPropertyDisplayName(Property)));
 							return true;
 						}
 					}
